Sequencer Playlists: Update UI to bring it in line with the latest designs.

diff --git a/Engine/Plugins/VirtualProduction/SequencerPlaylists/Source/SequencerPlaylists/Private/SequencerPlaylistItem_Sequence.h b/Engine/Plugins/VirtualProduction/SequencerPlaylists/Source/SequencerPlaylists/Private/SequencerPlaylistItem_Sequence.h
--- a/Engine/Plugins/VirtualProduction/SequencerPlaylists/Source/SequencerPlaylists/Private/SequencerPlaylistItem_Sequence.h
+++ b/Engine/Plugins/VirtualProduction/SequencerPlaylists/Source/SequencerPlaylists/Private/SequencerPlaylistItem_Sequence.h
@@ -22,7 +22,7 @@ class USequencerPlaylistItem_Sequence : public USequencerPlaylistItem
 	FText GetDisplayName() override;
 
 public:
-	UPROPERTY(EditAnywhere, Category="SequencerPlaylists")
+	UPROPERTY(EditAnywhere, Category="SequencerPlaylists", meta=(NoResetToDefault))
 	TObjectPtr<ULevelSequence> Sequence;
 };
 
diff --git a/Engine/Plugins/VirtualProduction/SequencerPlaylists/Source/SequencerPlaylists/Private/SequencerPlaylistPlayer.cpp b/Engine/Plugins/VirtualProduction/SequencerPlaylists/Source/SequencerPlaylists/Private/SequencerPlaylistPlayer.cpp
--- a/Engine/Plugins/VirtualProduction/SequencerPlaylists/Source/SequencerPlaylists/Private/SequencerPlaylistPlayer.cpp
+++ b/Engine/Plugins/VirtualProduction/SequencerPlaylists/Source/SequencerPlaylists/Private/SequencerPlaylistPlayer.cpp
@@ -92,7 +92,7 @@ bool USequencerPlaylistPlayer::PlayAll()
 	FScopedTransaction Transaction(LOCTEXT("PlayAllTransaction", "Trigger playback of all items"));
 	for (USequencerPlaylistItem* Item : Playlist->Items)
 	{
-		bResult |= GetCheckedItemPlayer(Item)->Play(Item);
+		bResult &= GetCheckedItemPlayer(Item)->Play(Item);
 	}
 
 	return bResult;
@@ -111,7 +111,7 @@ bool USequencerPlaylistPlayer::StopAll()
 	FScopedTransaction Transaction(LOCTEXT("StopAllTransaction", "Stop playback of all items"));
 	for (USequencerPlaylistItem* Item : Playlist->Items)
 	{
-		bResult |= GetCheckedItemPlayer(Item)->Stop(Item);
+		bResult &= GetCheckedItemPlayer(Item)->Stop(Item);
 	}
 
 	return bResult;
@@ -130,7 +130,7 @@ bool USequencerPlaylistPlayer::ResetAll()
 	FScopedTransaction Transaction(LOCTEXT("ResetAllTransaction", "Reset playback of all items"));
 	for (USequencerPlaylistItem* Item : Playlist->Items)
 	{
-		bResult |= GetCheckedItemPlayer(Item)->Reset(Item);
+		bResult &= GetCheckedItemPlayer(Item)->Reset(Item);
 	}
 
 	return bResult;
diff --git a/Engine/Plugins/VirtualProduction/SequencerPlaylists/Source/SequencerPlaylists/Private/SequencerPlaylistsModule.cpp b/Engine/Plugins/VirtualProduction/SequencerPlaylists/Source/SequencerPlaylists/Private/SequencerPlaylistsModule.cpp
--- a/Engine/Plugins/VirtualProduction/SequencerPlaylists/Source/SequencerPlaylists/Private/SequencerPlaylistsModule.cpp
+++ b/Engine/Plugins/VirtualProduction/SequencerPlaylists/Source/SequencerPlaylists/Private/SequencerPlaylistsModule.cpp
@@ -54,7 +54,8 @@ void FSequencerPlaylistsModule::StartupModule()
 	FGlobalTabmanager::Get()->RegisterNomadTabSpawner(SequencerPlaylistsTabName, FOnSpawnTab::CreateRaw(this, &FSequencerPlaylistsModule::OnSpawnPluginTab))
 		.SetDisplayName(LOCTEXT("SequencerPlaylistsTabTitle", "Playlists"))
 		.SetTooltipText(LOCTEXT("SequencerPlaylistsTabTooltipText", "Open the Sequencer Playlists tab."))
-		.SetGroup(WorkspaceMenu::GetMenuStructure().GetLevelEditorVirtualProductionCategory());
+		.SetGroup(WorkspaceMenu::GetMenuStructure().GetLevelEditorVirtualProductionCategory())
+		.SetIcon(FSlateIcon(FSequencerPlaylistsStyle::GetStyleSetName(), "SequencerPlaylists.TabIcon"));
 }
 
 
diff --git a/Engine/Plugins/VirtualProduction/SequencerPlaylists/Source/SequencerPlaylists/Private/SequencerPlaylistsModule.h b/Engine/Plugins/VirtualProduction/SequencerPlaylists/Source/SequencerPlaylists/Private/SequencerPlaylistsModule.h
--- a/Engine/Plugins/VirtualProduction/SequencerPlaylists/Source/SequencerPlaylists/Private/SequencerPlaylistsModule.h
+++ b/Engine/Plugins/VirtualProduction/SequencerPlaylists/Source/SequencerPlaylists/Private/SequencerPlaylistsModule.h
@@ -29,6 +29,8 @@ public:
 
 	TSharedPtr<ISequencerPlaylistItemPlayer> CreateItemPlayerForClass(TSubclassOf<USequencerPlaylistItem> ItemClass, TSharedRef<ISequencer> Sequencer);
 
+	TSharedPtr<FUICommandList> GetCommandList() { return PluginCommands; }
+
 private:
 	/** This function will be bound to Command (by default it will bring up plugin window) */
 	void PluginButtonClicked();
diff --git a/Engine/Plugins/VirtualProduction/SequencerPlaylists/Source/SequencerPlaylists/Private/SequencerPlaylistsStyle.cpp b/Engine/Plugins/VirtualProduction/SequencerPlaylists/Source/SequencerPlaylists/Private/SequencerPlaylistsStyle.cpp
--- a/Engine/Plugins/VirtualProduction/SequencerPlaylists/Source/SequencerPlaylists/Private/SequencerPlaylistsStyle.cpp
+++ b/Engine/Plugins/VirtualProduction/SequencerPlaylists/Source/SequencerPlaylists/Private/SequencerPlaylistsStyle.cpp
@@ -6,15 +6,12 @@
 #include "Slate/SlateGameResources.h"
 #include "Interfaces/IPluginManager.h"
 #include "Styling/SlateStyleMacros.h"
+#include "Styling/StyleColors.h"
 
 #define RootToContentDir Style->RootToContentDir
 
 TSharedPtr<FSlateStyleSet> FSequencerPlaylistsStyle::StyleInstance = nullptr;
 
-// FIXME: These should come from IAssetTypeActions::GetTypeColor()
-const FSlateColor FSequencerPlaylistsStyle::StyleColors::Animation(FColor::FromHex("143210"));
-const FSlateColor FSequencerPlaylistsStyle::StyleColors::Sequence(FColor::FromHex("931414"));
-
 
 void FSequencerPlaylistsStyle::Initialize()
 {
@@ -38,28 +35,115 @@ FName FSequencerPlaylistsStyle::GetStyleSetName()
 	return StyleSetName;
 }
 
-TSharedRef< FSlateStyleSet > FSequencerPlaylistsStyle::Create()
+TSharedRef<FSlateStyleSet> FSequencerPlaylistsStyle::Create()
 {
 	static const FVector2D Icon16x16(16.0f, 16.0f);
 	static const FVector2D Icon20x20(20.0f, 20.0f);
+	static const FVector2D Icon24x24(24.0f, 24.0f);
+	static const FVector2D LargeTransport(60.0f, 60.0f);
 
-	TSharedRef< FSlateStyleSet > Style = MakeShareable(new FSlateStyleSet("SequencerPlaylistsStyle"));
+	TSharedRef<FSlateStyleSet> Style = MakeShareable(new FSlateStyleSet("SequencerPlaylistsStyle"));
 	Style->SetContentRoot(IPluginManager::Get().FindPlugin("SequencerPlaylists")->GetBaseDir() / TEXT("Resources"));
 
-	Style->Set("SequencerPlaylists.OpenPluginWindow", new IMAGE_BRUSH_SVG(TEXT("PlaceholderButtonIcon"), Icon20x20));
-	Style->Set("SequencerPlaylists.AnimationBorderBrush", new FSlateColorBrush(StyleColors::Animation));
-	Style->Set("SequencerPlaylists.SequenceBorderBrush", new FSlateColorBrush(StyleColors::Sequence));
+	Style->Set("SequencerPlaylists.TabIcon", new IMAGE_BRUSH_SVG("PlaylistIcon", Icon16x16));
+	Style->Set("SequencerPlaylists.OpenPluginWindow", new IMAGE_BRUSH_SVG("PlaceholderButtonIcon", Icon20x20));
+	Style->Set("SequencerPlaylists.Panel.Background", new FSlateColorBrush(FStyleColors::Panel));
+
+	Style->Set("SequencerPlaylists.NewPlaylist.Background", new IMAGE_BRUSH_SVG("PlaylistIcon", Icon20x20));
+	// NewPlaylist.Overlay, SavePlaylistAs, OpenPlaylist are registered under a different root below.
+
+	const FButtonStyle& SimpleButtonStyle = Style->GetWidgetStyle<FButtonStyle>("SimpleButton");
+
+	FButtonStyle TransportButtonStyle = FButtonStyle(SimpleButtonStyle)
+		.SetNormalForeground(FStyleColors::Foreground)
+		.SetPressedForeground(FStyleColors::ForegroundHover)
+		.SetHoveredForeground(FStyleColors::ForegroundHover)
+		.SetDisabledForeground(FStyleColors::Foreground);
+
+	const float CornerRadius = 4.0f;
+	const FVector4 LeftCorners(CornerRadius, 0.0f, 0.0f, CornerRadius);
+	const FVector4 MiddleCorners(0.0f, 0.0f, 0.0f, 0.0f);
+	const FVector4 RightCorners(0.0f, CornerRadius, CornerRadius, 0.0f);
+
+	const FSlateRoundedBoxBrush LeftBrush(FStyleColors::Dropdown, LeftCorners);
+	const FSlateRoundedBoxBrush MiddleBrush(FStyleColors::Dropdown, MiddleCorners);
+	const FSlateRoundedBoxBrush RightBrush(FStyleColors::Dropdown, RightCorners);
+
+	const FButtonStyle LeftTransportButton = FButtonStyle(TransportButtonStyle)
+		.SetNormal(LeftBrush)
+		.SetHovered(LeftBrush)
+		.SetPressed(LeftBrush)
+		.SetDisabled(LeftBrush)
+		.SetNormalPadding(FMargin(8.f, 4.f, 6.f, 4.f))
+		.SetPressedPadding(FMargin(8.f, 4.f, 6.f, 4.f));
+
+	const FButtonStyle MiddleTransportButton = FButtonStyle(TransportButtonStyle)
+		.SetNormal(MiddleBrush)
+		.SetHovered(MiddleBrush)
+		.SetPressed(MiddleBrush)
+		.SetDisabled(MiddleBrush)
+		.SetNormalPadding(FMargin(6.f, 4.f, 6.f, 4.f))
+		.SetPressedPadding(FMargin(6.f, 4.f, 6.f, 4.f));
+
+	const FButtonStyle RightTransportButton = FButtonStyle(TransportButtonStyle)
+		.SetNormal(RightBrush)
+		.SetHovered(RightBrush)
+		.SetPressed(RightBrush)
+		.SetDisabled(RightBrush)
+		.SetNormalPadding(FMargin(6.f, 4.f, 8.f, 4.f))
+		.SetPressedPadding(FMargin(6.f, 4.f, 8.f, 4.f));
+
+	const FButtonStyle PlayTransportButton = FButtonStyle(LeftTransportButton)
+		.SetHoveredForeground(FStyleColors::AccentGreen);
+
+	const FButtonStyle StopTransportButton = FButtonStyle(MiddleTransportButton)
+		.SetHoveredForeground(FStyleColors::AccentRed);
+
+	const FButtonStyle ResetTransportButton = FButtonStyle(RightTransportButton)
+		.SetHoveredForeground(FStyleColors::AccentBlue);
+
+	Style->Set("SequencerPlaylists.TransportButton.Play", PlayTransportButton);
+	Style->Set("SequencerPlaylists.TransportButton.Stop", StopTransportButton);
+	Style->Set("SequencerPlaylists.TransportButton.Reset", ResetTransportButton);
+
+	FEditableTextBoxStyle EditableTextStyle = FEditableTextBoxStyle()
+		.SetFont(FCoreStyle::GetDefaultFontStyle("Regular", 9))
+		.SetBackgroundImageNormal(FSlateNoResource())
+		.SetBackgroundImageHovered(FSlateNoResource())
+		.SetBackgroundImageFocused(FSlateNoResource())
+		.SetBackgroundImageReadOnly(FSlateNoResource())
+		.SetBackgroundColor(FLinearColor::Transparent)
+		.SetForegroundColor(FSlateColor::UseForeground());
+
+	Style->Set("SequencerPlaylists.EditableTextBox", EditableTextStyle);
+	Style->Set("SequencerPlaylists.TitleFont", FCoreStyle::GetDefaultFontStyle("Bold", 18));
+	Style->Set("SequencerPlaylists.DescriptionFont", FCoreStyle::GetDefaultFontStyle("Regular", 8));
 
 	{
-		// These assets live in the editor slate folder, change the root temporarily
+		// These assets live in engine slate folders; change the root temporarily
 		const FString PreviousContentRoot = Style->GetContentRootDir();
 		ON_SCOPE_EXIT{ Style->SetContentRoot(PreviousContentRoot); };
 
-		const FString EngineEditorSlateDir = FPaths::EngineContentDir() / TEXT("Editor/Slate");
-		Style->SetContentRoot(EngineEditorSlateDir);
+		// Engine/Content/Slate/...
+		{
+			const FString EngineSlateDir = FPaths::EngineContentDir() / TEXT("Slate");
+			Style->SetContentRoot(EngineSlateDir);
+
+			Style->Set("SequencerPlaylists.Ellipsis", new IMAGE_BRUSH_SVG("Starship/Common/ellipsis-vertical-narrow", FVector2D(6, 24)));
+		}
+
+		// Engine/Content/Editor/Slate/...
+		{
+			const FString EngineEditorSlateDir = FPaths::EngineContentDir() / TEXT("Editor") / TEXT("Slate");
+			Style->SetContentRoot(EngineEditorSlateDir);
+
+			Style->Set("SequencerPlaylists.NewPlaylist.Overlay", new IMAGE_BRUSH_SVG("Starship/MainToolbar/ToolBadgePlus", Icon20x20, FStyleColors::AccentGreen));
+			Style->Set("SequencerPlaylists.SavePlaylistAs", new IMAGE_BRUSH_SVG("Starship/Common/SaveCurrentAs", Icon16x16));
+			Style->Set("SequencerPlaylists.OpenPlaylist", new IMAGE_BRUSH_SVG("Starship/Common/OpenAsset", Icon16x16));
 
-		Style->Set("SequencerPlaylists.DropZone.Above", new BOX_BRUSH("Common/DropZoneIndicator_Above", FMargin(10.0f / 16.0f, 10.0f / 16.0f, 0.f, 0.f), Style->GetSlateColor("SelectionColor")));
-		Style->Set("SequencerPlaylists.DropZone.Below", new BOX_BRUSH("Common/DropZoneIndicator_Below", FMargin(10.0f / 16.0f, 0.f, 0.f, 10.0f / 16.0f), Style->GetSlateColor("SelectionColor")));
+			Style->Set("SequencerPlaylists.Play", new IMAGE_BRUSH_SVG("Starship/Common/play", Icon20x20));
+			Style->Set("SequencerPlaylists.Stop", new IMAGE_BRUSH_SVG("Starship/MainToolbar/stop", Icon20x20));
+		}
 	}
 
 	return Style;
diff --git a/Engine/Plugins/VirtualProduction/SequencerPlaylists/Source/SequencerPlaylists/Private/SequencerPlaylistsStyle.h b/Engine/Plugins/VirtualProduction/SequencerPlaylists/Source/SequencerPlaylists/Private/SequencerPlaylistsStyle.h
--- a/Engine/Plugins/VirtualProduction/SequencerPlaylists/Source/SequencerPlaylists/Private/SequencerPlaylistsStyle.h
+++ b/Engine/Plugins/VirtualProduction/SequencerPlaylists/Source/SequencerPlaylists/Private/SequencerPlaylistsStyle.h
@@ -9,12 +9,6 @@
 class FSequencerPlaylistsStyle
 {
 public:
-	struct StyleColors
-	{
-		static const FSlateColor Animation;
-		static const FSlateColor Sequence;
-	};
-
 	static void Initialize();
 	static void Shutdown();
 
diff --git a/Engine/Plugins/VirtualProduction/SequencerPlaylists/Source/SequencerPlaylists/Private/SequencerPlaylistsWidgets.cpp b/Engine/Plugins/VirtualProduction/SequencerPlaylists/Source/SequencerPlaylists/Private/SequencerPlaylistsWidgets.cpp
--- a/Engine/Plugins/VirtualProduction/SequencerPlaylists/Source/SequencerPlaylists/Private/SequencerPlaylistsWidgets.cpp
+++ b/Engine/Plugins/VirtualProduction/SequencerPlaylists/Source/SequencerPlaylists/Private/SequencerPlaylistsWidgets.cpp
@@ -5,25 +5,34 @@
 #include "SequencerPlaylistItem.h"
 #include "SequencerPlaylistItem_Sequence.h"
 #include "SequencerPlaylistPlayer.h"
+#include "SequencerPlaylistsModule.h"
 #include "SequencerPlaylistsStyle.h"
 
 #include "AssetData.h"
 #include "AssetRegistryModule.h"
 #include "AssetToolsModule.h"
 #include "ContentBrowserModule.h"
+#include "EditorFontGlyphs.h"
 #include "FileHelpers.h"
 #include "Framework/MultiBox/MultiBoxBuilder.h"
 #include "IContentBrowserSingleton.h"
 #include "ISinglePropertyView.h"
 #include "LevelEditor.h"
 #include "Misc/FileHelper.h"
+#include "Misc/TextFilter.h"
 #include "ScopedTransaction.h"
 #include "SlateOptMacros.h"
+#include "SPositiveActionButton.h"
 #include "Styling/SlateIconFinder.h"
+#include "Styling/StyleColors.h"
+#include "Widgets/SOverlay.h"
 #include "Widgets/Images/SImage.h"
 #include "Widgets/Input/SButton.h"
 #include "Widgets/Input/SCheckBox.h"
 #include "Widgets/Input/SComboButton.h"
+#include "Widgets/Input/SMenuAnchor.h"
+#include "Widgets/Input/SMultiLineEditableTextBox.h"
+#include "Widgets/Input/SSearchBox.h"
 #include "Widgets/Layout/SScrollBox.h"
 #include "Widgets/Layout/SSpacer.h"
 
@@ -31,7 +40,13 @@
 #define LOCTEXT_NAMESPACE "SequencerPlaylists"
 
 
-const float SSequencerPlaylistPanel::DefaultWidth = 850.0f;
+const float SSequencerPlaylistPanel::DefaultWidth = 300.0f;
+const FName SSequencerPlaylistPanel::ColumnName_HoverTransport(TEXT("HoverTransport"));
+const FName SSequencerPlaylistPanel::ColumnName_Items(TEXT("Items"));
+const FName SSequencerPlaylistPanel::ColumnName_Offset(TEXT("Offset"));
+const FName SSequencerPlaylistPanel::ColumnName_Hold(TEXT("Hold"));
+const FName SSequencerPlaylistPanel::ColumnName_Loop(TEXT("Loop"));
+const FName SSequencerPlaylistPanel::ColumnName_HoverDetails(TEXT("HoverDetails"));
 
 
 BEGIN_SLATE_FUNCTION_BUILD_OPTIMIZATION
@@ -40,181 +55,312 @@ void SSequencerPlaylistPanel::Construct(const FArguments& InArgs, USequencerPlay
 	check(InPlayer);
 	WeakPlayer = InPlayer;
 
-	const ISlateStyle& SpmStyle = FSequencerPlaylistsStyle::Get();
-
-	SAssignNew(ItemList, SDragAndDropVerticalBox)
-		.OnCanAcceptDrop(this, &SSequencerPlaylistPanel::HandleCanAcceptDrop)
-		.OnAcceptDrop(this, &SSequencerPlaylistPanel::HandleAcceptDrop)
-		.OnDragDetected(this, &SSequencerPlaylistPanel::HandleDragDetected);
-
-	(*ItemList)
-		.SetDropIndicator_Above(*SpmStyle.GetBrush("SequencerPlaylists.DropZone.Above"))
-		.SetDropIndicator_Below(*SpmStyle.GetBrush("SequencerPlaylists.DropZone.Below"));
-
-	RegenerateSequenceList();
+	SearchTextFilter = MakeShared<TTextFilter<const FSequencerPlaylistRowData&>>(
+		TTextFilter<const FSequencerPlaylistRowData&>::FItemToStringArray::CreateSP(this, &SSequencerPlaylistPanel::GetSearchStrings));
+	SearchTextFilter->OnChanged().AddSP(this, &SSequencerPlaylistPanel::RegenerateRows);
 
 	ChildSlot
 	[
-		SNew(SScrollBox)
-		+ SScrollBox::Slot()
+		SNew(SOverlay)
+		+ SOverlay::Slot()
+		[
+			SNew(SImage)
+			.Image(FSequencerPlaylistsStyle::Get().GetBrush("SequencerPlaylists.Panel.Background"))
+		]
+		+ SOverlay::Slot()
 		[
-			SNew(SComboButton)
-			.ContentPadding(4.f)
-			.OnGetMenuContent(this, &SSequencerPlaylistPanel::OnPresetGeneratePresetsMenu)
-			.ButtonContent()
+			SNew(SVerticalBox)
+			+SVerticalBox::Slot()
+			.AutoHeight()
+			[
+				Construct_Toolbar()
+			]
+			+ SVerticalBox::Slot()
+			.AutoHeight()
+			.Padding(8.0f, 3.0f)
+			[
+				SNew(STextBlock)
+				.Font(FSequencerPlaylistsStyle::Get().GetFontStyle("SequencerPlaylists.TitleFont"))
+				.Text(TAttribute<FText>::CreateLambda([This = SharedThis(this)]()
+					{
+						if (USequencerPlaylist* Playlist = This->GetCheckedPlaylist())
+						{
+							return FText::FromString(Playlist->GetName());
+						}
+
+						return FText::GetEmpty();
+					}))
+			]
+			+ SVerticalBox::Slot()
+			.AutoHeight()
+			.MaxHeight(150.0f)
+			.Padding(8.0f, 3.0f)
 			[
 				SNew(SHorizontalBox)
 				+ SHorizontalBox::Slot()
-				.Padding(4.f, 0.f, 4.f, 0.f)
-				.AutoWidth()
+				.FillWidth(1.0f)
 				[
-					SNew(SImage)
-					.Image(FSlateIconFinder::FindIconBrushForClass(USequencerPlaylist::StaticClass()))
+					SNew(SScrollBox)
+					+ SScrollBox::Slot()
+					.Padding(0.0f, 0.0f, 8.0f, 0.0f)
+					[
+						SNew(SMultiLineEditableTextBox)
+						.Padding(0)
+						.Margin(0)
+						.AutoWrapText(true)
+						.Style(FSequencerPlaylistsStyle::Get(), "SequencerPlaylists.EditableTextBox")
+						.Font(FSequencerPlaylistsStyle::Get().GetFontStyle("SequencerPlaylists.DescriptionFont"))
+						.HintText(LOCTEXT("PlaylistDescriptionHint", "<playlist description>"))
+						.Text(TAttribute<FText>::CreateLambda([This = SharedThis(this)]()
+							{
+								if (USequencerPlaylist* Playlist = This->GetCheckedPlaylist())
+								{
+									return Playlist->Description;
+								}
+
+								return FText::GetEmpty();
+							}))
+						.OnTextCommitted_Lambda([This = SharedThis(this)](const FText& NewText, ETextCommit::Type CommitType)
+							{
+								if (USequencerPlaylist* Playlist = This->GetCheckedPlaylist())
+								{
+									Playlist->Description = NewText;
+								}
+							})
+					]
 				]
 				+ SHorizontalBox::Slot()
+				.AutoWidth()
+				.VAlign(VAlign_Bottom)
 				[
-					SNew(STextBlock)
-					.Text(LOCTEXT("PresetsToolbarButton", "Presets"))
+					Construct_Transport()
 				]
 			]
-		]
-		+ SScrollBox::Slot()
-		[
-			SNew(SSpacer)
-			.Size(FVector2D(1.0f, 10.0f))
-		]
-		+ SScrollBox::Slot()
-		.Padding(10.0f)
-		[
-			// Play all
-			SNew(SHorizontalBox)
-			+ SHorizontalBox::Slot()
-			.FillWidth(0.04f)
+			+ SVerticalBox::Slot()
+			.AutoHeight()
+			.Padding(8.0f, 6.0f, 8.0f, 3.0f)
 			[
-				SNew(SBorder)
-				.BorderImage(SpmStyle.GetBrush("SequencerPlaylists.SequenceBorderBrush"))
+				Construct_AddSearchRow()
 			]
-			+ SHorizontalBox::Slot()
-			.FillWidth(1.0f)
-			.Padding(4.0f, 0.0f, 0.0f, 0.0f)
+			+ SVerticalBox::Slot()
+			.FillHeight(1.0f)
 			[
-				SNew(SVerticalBox)
-				+ SVerticalBox::Slot()
-				[
-					SNew(SButton)
-					.Text(LOCTEXT("PlayAllButton", "Play All"))
-					.HAlign(HAlign_Center)
-					.OnClicked(this, &SSequencerPlaylistPanel::OnClicked_PlayAll)
-				]
-				+ SVerticalBox::Slot()
-				[
-					SNew(SHorizontalBox)
-					+ SHorizontalBox::Slot()
-					.AutoWidth()
-					.VAlign(VAlign_Center)
-					.Padding(2.0f, 1.0f, 0.0f, 0.0f)
-					[
-						SNew(STextBlock)
-						.Text(LOCTEXT("HoldAllAtFirstFrameCheckLabel", "Hold All at First Frame"))
-					]
-					+ SHorizontalBox::Slot()
-					.AutoWidth()
-					.Padding(4.0f, 1.0f, 10.0f, 0.0f)
-					[
-						SNew(SCheckBox)
-						.IsChecked(ECheckBoxState::Undetermined)
-					]
-					+ SHorizontalBox::Slot()
-					.FillWidth(1.0f)
-					[
-						SNew(SButton)
-						.Text(LOCTEXT("StopAllButton", "Stop All"))
-						.HAlign(HAlign_Center)
-						.OnClicked(this, &SSequencerPlaylistPanel::OnClicked_StopAll)
-					]
-					+ SHorizontalBox::Slot()
-					.FillWidth(1.0f)
-					[
-						SNew(SButton)
-						.Text(LOCTEXT("ResetAllButton", "Reset All"))
-						.HAlign(HAlign_Center)
-						.OnClicked(this, &SSequencerPlaylistPanel::OnClicked_ResetAll)
-					]
-				]
+				Construct_ItemListView()
 			]
 		]
-		+ SScrollBox::Slot()
+	];
+
+	RegenerateRows();
+}
+END_SLATE_FUNCTION_BUILD_OPTIMIZATION
+
+
+BEGIN_SLATE_FUNCTION_BUILD_OPTIMIZATION
+TSharedRef<SWidget> SSequencerPlaylistPanel::Construct_Toolbar()
+{
+	FSequencerPlaylistsModule& Module = static_cast<FSequencerPlaylistsModule&>(ISequencerPlaylistsModule::Get());
+	FSlimHorizontalToolBarBuilder ToolBarBuilder(Module.GetCommandList(), FMultiBoxCustomization::None, nullptr, true);
+
+	ToolBarBuilder.BeginSection("Playlists");
+	{
+		ToolBarBuilder.AddToolBarButton(
+			FUIAction(FExecuteAction::CreateSP(this, &SSequencerPlaylistPanel::OnNewPlaylist)),
+			NAME_None,
+			LOCTEXT("NewPlaylistOptions", "New Playlist"),
+			LOCTEXT("NewPlaylistOptionsToolTip", "New Playlist"),
+			FSlateIcon(FSequencerPlaylistsStyle::Get().GetStyleSetName(), "SequencerPlaylists.NewPlaylist.Background", NAME_None, "SequencerPlaylists.NewPlaylist.Overlay"));
+
+		ToolBarBuilder.AddToolBarButton(
+			FUIAction(FExecuteAction::CreateSP(this, &SSequencerPlaylistPanel::OnSavePlaylistAs)),
+			NAME_None,
+			LOCTEXT("SavePlaylist", "Save Playlist As..."),
+			LOCTEXT("SavePlaylistTooltip", "Save Playlist As..."),
+			FSlateIcon(FSequencerPlaylistsStyle::Get().GetStyleSetName(), "SequencerPlaylists.SavePlaylistAs"));
+
+		ToolBarBuilder.AddComboButton(
+			FUIAction(),
+			FOnGetContent::CreateSP(SharedThis(this), &SSequencerPlaylistPanel::BuildOpenPlaylistMenu),
+			LOCTEXT("NewPlaylistOptions", "Open Playlist"),
+			LOCTEXT("NewPlaylistOptionsToolTip", "Open Playlist"),
+			FSlateIcon(FSequencerPlaylistsStyle::Get().GetStyleSetName(), "SequencerPlaylists.OpenPlaylist"));
+	}
+	ToolBarBuilder.EndSection();
+
+	return ToolBarBuilder.MakeWidget();
+}
+END_SLATE_FUNCTION_BUILD_OPTIMIZATION
+
+
+BEGIN_SLATE_FUNCTION_BUILD_OPTIMIZATION
+TSharedRef<SWidget> SSequencerPlaylistPanel::Construct_Transport()
+{
+	return SNew(SHorizontalBox)
+		+ SHorizontalBox::Slot()
+		.AutoWidth()
 		[
-			SNew(SSpacer)
-			.Size(FVector2D(1.0f, 10.0f))
+			SNew(SButton)
+			.ButtonStyle(FSequencerPlaylistsStyle::Get(), "SequencerPlaylists.TransportButton.Play")
+			.ContentPadding(0)
+			.HAlign(HAlign_Center)
+			.VAlign(VAlign_Center)
+			.OnClicked(this, &SSequencerPlaylistPanel::HandleClicked_PlayAll)
+			.ToolTipText(LOCTEXT("PlayAllButtonToolTip", "Play all items"))
+			[
+				SNew(SImage)
+				.Image(FSequencerPlaylistsStyle::Get().GetBrush("SequencerPlaylists.Play"))
+				.ColorAndOpacity(FSlateColor::UseForeground())
+			]
 		]
-		+ SScrollBox::Slot()
+		+ SHorizontalBox::Slot()
+		.AutoWidth()
 		[
-			ItemList.ToSharedRef()
+			SNew(SButton)
+			.ButtonStyle(FSequencerPlaylistsStyle::Get(), "SequencerPlaylists.TransportButton.Stop")
+			.ContentPadding(0)
+			.HAlign(HAlign_Center)
+			.VAlign(VAlign_Center)
+			.OnClicked(this, &SSequencerPlaylistPanel::HandleClicked_StopAll)
+			.ToolTipText(LOCTEXT("StopAllButtonToolTip", "Stop all items"))
+			[
+				SNew(SImage)
+				.Image(FSequencerPlaylistsStyle::Get().GetBrush("SequencerPlaylists.Stop"))
+				.ColorAndOpacity(FSlateColor::UseForeground())
+			]
 		]
-		+ SScrollBox::Slot()
-		.Padding(10.0f)
+		+ SHorizontalBox::Slot()
+		.AutoWidth()
 		[
-			SNew(SHorizontalBox)
-			+ SHorizontalBox::Slot()
-			.FillWidth(1.0f)
+			SNew(SButton)
+			.ButtonStyle(FSequencerPlaylistsStyle::Get(), "SequencerPlaylists.TransportButton.Reset")
+			.ContentPadding(FMargin(0.0f, 1.0f, 0.0f, 0.0f))
+			.HAlign(HAlign_Center)
+			.VAlign(VAlign_Center)
+			.OnClicked(this, &SSequencerPlaylistPanel::HandleClicked_ResetAll)
+			.ToolTipText(LOCTEXT("ResetAllButtonToolTip", "Reset all items"))
 			[
-				SNew(SButton)
-				.Text(LOCTEXT("AddSequenceButton", "Add Sequence"))
-				.HAlign(HAlign_Center)
-				.OnClicked(this, &SSequencerPlaylistPanel::OnClicked_AddSequence)
+				SNew(STextBlock)
+				.Font(FEditorStyle::Get().GetFontStyle("FontAwesome.14"))
+				.Text(FEditorFontGlyphs::Undo)
 			]
+		];
+}
+END_SLATE_FUNCTION_BUILD_OPTIMIZATION
+
+
+BEGIN_SLATE_FUNCTION_BUILD_OPTIMIZATION
+TSharedRef<SWidget> SSequencerPlaylistPanel::Construct_AddSearchRow()
+{
+	return SNew(SHorizontalBox)
+		+ SHorizontalBox::Slot()
+		.AutoWidth()
+		[
+			SNew(SPositiveActionButton)
+			.OnClicked(this, &SSequencerPlaylistPanel::HandleClicked_AddSequence)
+			.Icon(FAppStyle::Get().GetBrush("Icons.Plus"))
+			.Text(LOCTEXT("AddItemButton", "Item"))
 		]
-	];
+		+ SHorizontalBox::Slot()
+		.FillWidth(1.0f)
+		.Padding(8.0f, 0.0f)
+		[
+			SAssignNew(SearchBox, SSearchBox)
+			.HintText(LOCTEXT("SearchHint", "Search Playlist Items"))
+			.OnTextChanged(this, &SSequencerPlaylistPanel::OnSearchTextChanged)
+			.DelayChangeNotificationsWhileTyping(true)
+		];
 }
 END_SLATE_FUNCTION_BUILD_OPTIMIZATION
 
 
 BEGIN_SLATE_FUNCTION_BUILD_OPTIMIZATION
-void SSequencerPlaylistPanel::RegenerateSequenceList()
+TSharedRef<SWidget> SSequencerPlaylistPanel::Construct_ItemListView()
 {
-	ItemList->ClearChildren();
+	return SAssignNew(ItemListView, SListView<TSharedPtr<FSequencerPlaylistRowData>>)
+		.SelectionMode(ESelectionMode::Single) // See TODO in HandleAcceptDrop
+		.ListItemsSource(&ItemRows)
+		.OnGenerateRow_Lambda([This = SharedThis(this)](TSharedPtr<FSequencerPlaylistRowData> InData, const TSharedRef<STableViewBase>& OwnerTableView)
+			{
+				return SNew(SSequencerPlaylistItemWidget, InData, OwnerTableView)
+					.OnPlayClicked(This, &SSequencerPlaylistPanel::HandleClicked_Item_Play)
+					.OnStopClicked(This, &SSequencerPlaylistPanel::HandleClicked_Item_Stop)
+					.OnResetClicked(This, &SSequencerPlaylistPanel::HandleClicked_Item_Reset)
+					.OnRemoveClicked(This, &SSequencerPlaylistPanel::HandleClicked_Item_Remove )
+					.OnIsPropertyVisible(This, &SSequencerPlaylistPanel::HandleItemDetailsIsPropertyVisible)
+					.OnCanAcceptDrop(This, &SSequencerPlaylistPanel::HandleCanAcceptDrop)
+					.OnAcceptDrop(This, &SSequencerPlaylistPanel::HandleAcceptDrop);
+			})
+		.HeaderRow(
+			SNew(SHeaderRow)
+			+ SHeaderRow::Column(ColumnName_HoverTransport)
+				.DefaultLabel(FText::GetEmpty())
+				.FillSized(20.0f)
+			+ SHeaderRow::Column(ColumnName_Items)
+				.DefaultLabel(LOCTEXT("ColumnLabelItems", "Playlist Items"))
+				.FillWidth(1.0f)
+			+ SHeaderRow::Column(ColumnName_Offset)
+				.DefaultLabel(LOCTEXT("ColumnLabelOffset", "Offset"))
+				.FillSized(45.0f)
+			+ SHeaderRow::Column(ColumnName_Hold)
+				.DefaultLabel(LOCTEXT("ColumnLabelHold", "Hold"))
+				.FillSized(35.0f)
+				.HAlignCell(HAlign_Center)
+			+ SHeaderRow::Column(ColumnName_Loop)
+				.DefaultLabel(LOCTEXT("ColumnLabelLoop", "Loop"))
+				.FillSized(40.0f)
+			+ SHeaderRow::Column(ColumnName_HoverDetails)
+				.DefaultLabel(FText::GetEmpty())
+				.FillSized(20.0f)
+				.HAlignCell(HAlign_Center)
+				.VAlignCell(VAlign_Center)
+		);
+}
+END_SLATE_FUNCTION_BUILD_OPTIMIZATION
 
+
+USequencerPlaylist* SSequencerPlaylistPanel::GetCheckedPlaylist()
+{
 	USequencerPlaylistPlayer* Player = WeakPlayer.Get();
-	if (!Player)
+	if (!ensure(Player))
 	{
-		return;
+		return nullptr;
 	}
 
 	USequencerPlaylist* Playlist = Player->GetPlaylist();
+	ensure(Playlist);
+	return Playlist;
+}
+
+
+void SSequencerPlaylistPanel::RegenerateRows()
+{
+	USequencerPlaylist* Playlist = GetCheckedPlaylist();
 	if (!Playlist)
 	{
 		return;
 	}
 
-	for (USequencerPlaylistItem* Item : Playlist->Items)
+	const int32 ItemCount = Playlist->Items.Num();
+	ItemRows.Empty(ItemCount);
+	for (int32 ItemIndex = 0; ItemIndex < ItemCount; ++ItemIndex)
 	{
-		ItemList->AddSlot()
-		[
-			SNew(SSequencerPlaylistItemWidget, Item)
-			.OnPlayClicked(SharedThis(this), &SSequencerPlaylistPanel::OnClicked_Item_Play)
-			.OnStopClicked(SharedThis(this), &SSequencerPlaylistPanel::OnClicked_Item_Stop)
-			.OnResetClicked(SharedThis(this), &SSequencerPlaylistPanel::OnClicked_Item_Reset)
-			.OnCloseClicked(SharedThis(this), &SSequencerPlaylistPanel::OnClicked_Item_Close)
-		];
+		USequencerPlaylistItem* Item = Playlist->Items[ItemIndex];
+		FSequencerPlaylistRowData Row(ItemIndex, Item);
+		if (SearchTextFilter->PassesFilter(Row))
+		{
+			ItemRows.Emplace(MakeShared<FSequencerPlaylistRowData>(Row));
+		}
 	}
+
+	ItemListView->RequestListRefresh();
 }
-END_SLATE_FUNCTION_BUILD_OPTIMIZATION
 
 
-TSharedRef<SWidget> SSequencerPlaylistPanel::OnPresetGeneratePresetsMenu()
+TSharedRef<SWidget> SSequencerPlaylistPanel::BuildOpenPlaylistMenu()
 {
 	FMenuBuilder MenuBuilder(true, nullptr);
 
 	IContentBrowserSingleton& ContentBrowser = FModuleManager::LoadModuleChecked<FContentBrowserModule>("ContentBrowser").Get();
 
-	MenuBuilder.AddMenuEntry(
-		LOCTEXT("SaveAsPreset_Text", "Save As Preset"),
-		LOCTEXT("SaveAsPreset_Tip", "Save the current setup as a new preset that can be loaded at a later date"),
-		FSlateIcon(FEditorStyle::Get().GetStyleSetName(), "AssetEditor.SaveAsset"),
-		FUIAction(FExecuteAction::CreateSP(this, &SSequencerPlaylistPanel::OnSaveAsPreset))
-	);
-
 	FAssetPickerConfig AssetPickerConfig;
 	{
 		AssetPickerConfig.SelectionMode = ESelectionMode::Single;
@@ -229,12 +375,12 @@ TSharedRef<SWidget> SSequencerPlaylistPanel::OnPresetGeneratePresetsMenu()
 		AssetPickerConfig.bShowTypeInColumnView = false;
 		AssetPickerConfig.bSortByPathInColumnView = false;
 
-		AssetPickerConfig.AssetShowWarningText = LOCTEXT("NoPresets_Warning", "No Presets Found");
+		AssetPickerConfig.AssetShowWarningText = LOCTEXT("OpenPlaylistNoAssetsWarning", "No Playlists Found");
 		AssetPickerConfig.Filter.ClassNames.Add(USequencerPlaylist::StaticClass()->GetFName());
-		AssetPickerConfig.OnAssetSelected = FOnAssetSelected::CreateSP(this, &SSequencerPlaylistPanel::OnLoadPreset);
+		AssetPickerConfig.OnAssetSelected = FOnAssetSelected::CreateSP(this, &SSequencerPlaylistPanel::OnLoadPlaylist);
 	}
 
-	MenuBuilder.BeginSection(NAME_None, LOCTEXT("LoadPreset_MenuSection", "Load Preset"));
+	MenuBuilder.BeginSection(NAME_None, LOCTEXT("OpenPlaylistMenuSection", "Load Playlist"));
 	{
 		TSharedRef<SWidget> PresetPicker = SNew(SBox)
 			.MinDesiredWidth(400.f)
@@ -259,7 +405,7 @@ static bool OpenSaveDialog(const FString& InDefaultPath, const FString& InNewNam
 		SaveAssetDialogConfig.DefaultAssetName = InNewNameSuggestion;
 		SaveAssetDialogConfig.AssetClassNames.Add(USequencerPlaylist::StaticClass()->GetFName());
 		SaveAssetDialogConfig.ExistingAssetPolicy = ESaveAssetDialogExistingAssetPolicy::AllowButWarn;
-		SaveAssetDialogConfig.DialogTitleOverride = LOCTEXT("SavePresetDialogTitle", "Save Sequencer Playlist Preset");
+		SaveAssetDialogConfig.DialogTitleOverride = LOCTEXT("SavePlaylistDialogTitle", "Save Sequencer Playlist");
 	}
 
 	FContentBrowserModule& ContentBrowserModule = FModuleManager::LoadModuleChecked<FContentBrowserModule>("ContentBrowser");
@@ -275,9 +421,10 @@ static bool OpenSaveDialog(const FString& InDefaultPath, const FString& InNewNam
 }
 
 
-bool GetSavePresetPackageName(FString& OutName)
+bool GetSavePlaylistPackageName(FString& OutName)
 {
-	//ULiveLinkSettings* ConfigSettings = GetMutableDefault<ULiveLinkSettings>();
+	// TODO
+	//USequencerPlaylistsSettings* ConfigSettings = GetMutableDefault<USequencerPlaylistsSettings>();
 
 	FDateTime Today = FDateTime::Now();
 
@@ -295,7 +442,7 @@ bool GetSavePresetPackageName(FString& OutName)
 	}
 
 	// determine default asset name
-	FString DefaultName = LOCTEXT("NewPresetDefaultName", "NewSequencerPlaylistPreset").ToString();
+	FString DefaultName = LOCTEXT("NewPlaylistDefaultName", "NewSequencerPlaylist").ToString();
 
 	FString UniquePackageName;
 	FString UniqueAssetName;
@@ -330,7 +477,7 @@ bool GetSavePresetPackageName(FString& OutName)
 }
 
 
-void SSequencerPlaylistPanel::OnSaveAsPreset()
+void SSequencerPlaylistPanel::OnSavePlaylistAs()
 {
 	USequencerPlaylist* Playlist = GetCheckedPlaylist();
 	if (!Playlist)
@@ -339,12 +486,12 @@ void SSequencerPlaylistPanel::OnSaveAsPreset()
 	}
 
 	FString PackageName;
-	if (!GetSavePresetPackageName(PackageName))
+	if (!GetSavePlaylistPackageName(PackageName))
 	{
 		return;
 	}
 
-	FScopedTransaction Transaction(LOCTEXT("SaveAsPreset", "Save As Preset"));
+	FScopedTransaction Transaction(LOCTEXT("SavePlaylistTransaction", "Save Sequencer Playlist"));
 
 	// Saving into a new package
 	const FString NewAssetName = FPackageName::GetLongPackageAssetName(PackageName);
@@ -353,6 +500,8 @@ void SSequencerPlaylistPanel::OnSaveAsPreset()
 
 	if (NewPlaylist)
 	{
+		NewPlaylist->Description = Playlist->Description;
+
 		for (USequencerPlaylistItem* Item : Playlist->Items)
 		{
 			NewPlaylist->Items.Add(DuplicateObject<USequencerPlaylistItem>(Item, NewPlaylist));
@@ -366,7 +515,7 @@ void SSequencerPlaylistPanel::OnSaveAsPreset()
 }
 
 
-void SSequencerPlaylistPanel::OnLoadPreset(const FAssetData& InPreset)
+void SSequencerPlaylistPanel::OnLoadPlaylist(const FAssetData& InPreset)
 {
 	FSlateApplication::Get().DismissAllMenus();
 
@@ -379,18 +528,57 @@ void SSequencerPlaylistPanel::OnLoadPreset(const FAssetData& InPreset)
 	USequencerPlaylist* ImportedPlaylist = CastChecked<USequencerPlaylist>(InPreset.GetAsset());
 	if (ImportedPlaylist)
 	{
-		FScopedTransaction Transaction(LOCTEXT("LoadPreset_Transaction", "Load Sequencer Playlist preset"));
+		FScopedTransaction Transaction(LOCTEXT("LoadPlaylistTransaction", "Load Sequencer Playlist"));
+
+		// FIXME: Within the transient package, which we'll never be loading from.
+		// This will go away when this is refactored into an asset editor. (UETOOL-4643)
+		Playlist->Rename(*ImportedPlaylist->GetName());
+
+		Playlist->Description = ImportedPlaylist->Description;
+
 		Playlist->Items.Empty();
 		for (USequencerPlaylistItem* Item : ImportedPlaylist->Items)
 		{
 			Playlist->Items.Add(DuplicateObject<USequencerPlaylistItem>(Item, Playlist));
 		}
-		RegenerateSequenceList();
+
+		RegenerateRows();
+	}
+}
+
+
+void SSequencerPlaylistPanel::OnNewPlaylist()
+{
+	USequencerPlaylist* Playlist = GetCheckedPlaylist();
+	if (!Playlist)
+	{
+		return;
+	}
+
+	Playlist->Rename(*MakeUniqueObjectName(GetTransientPackage(), USequencerPlaylist::StaticClass(), "UntitledPlaylist").ToString());
+	Playlist->Description = FText::GetEmpty();
+	Playlist->Items.Empty();
+	RegenerateRows();
+}
+
+
+void SSequencerPlaylistPanel::GetSearchStrings(const FSequencerPlaylistRowData& Row, TArray<FString>& OutSearchStrings)
+{
+	if (USequencerPlaylistItem* Item = Row.WeakItem.Get())
+	{
+		OutSearchStrings.Add(Item->GetDisplayName().ToString());
 	}
 }
 
 
-FReply SSequencerPlaylistPanel::OnClicked_PlayAll()
+void SSequencerPlaylistPanel::OnSearchTextChanged(const FText& InFilterText)
+{
+	SearchTextFilter->SetRawFilterText(InFilterText);
+	SearchBox->SetError(SearchTextFilter->GetFilterErrorText());
+}
+
+
+FReply SSequencerPlaylistPanel::HandleClicked_PlayAll()
 {
 	USequencerPlaylistPlayer* Player = WeakPlayer.Get();
 	if (!ensure(Player))
@@ -403,7 +591,7 @@ FReply SSequencerPlaylistPanel::OnClicked_PlayAll()
 }
 
 
-FReply SSequencerPlaylistPanel::OnClicked_StopAll()
+FReply SSequencerPlaylistPanel::HandleClicked_StopAll()
 {
 	USequencerPlaylistPlayer* Player = WeakPlayer.Get();
 	if (!ensure(Player))
@@ -416,7 +604,7 @@ FReply SSequencerPlaylistPanel::OnClicked_StopAll()
 }
 
 
-FReply SSequencerPlaylistPanel::OnClicked_ResetAll()
+FReply SSequencerPlaylistPanel::HandleClicked_ResetAll()
 {
 	USequencerPlaylistPlayer* Player = WeakPlayer.Get();
 	if (!ensure(Player))
@@ -429,7 +617,7 @@ FReply SSequencerPlaylistPanel::OnClicked_ResetAll()
 }
 
 
-FReply SSequencerPlaylistPanel::OnClicked_AddSequence()
+FReply SSequencerPlaylistPanel::HandleClicked_AddSequence()
 {
 	USequencerPlaylist* Playlist = GetCheckedPlaylist();
 	if (!Playlist)
@@ -438,13 +626,13 @@ FReply SSequencerPlaylistPanel::OnClicked_AddSequence()
 	}
 
 	Playlist->Items.Add(NewObject<USequencerPlaylistItem_Sequence>(Playlist));
-	RegenerateSequenceList();
+	RegenerateRows(); // TODO: Incremental view update? (UETOOL-4644)
 
 	return FReply::Handled();
 }
 
 
-FReply SSequencerPlaylistPanel::OnClicked_Item_Play(TSharedPtr<SSequencerPlaylistItemWidget> ItemWidget)
+FReply SSequencerPlaylistPanel::HandleClicked_Item_Play(TSharedPtr<SSequencerPlaylistItemWidget> ItemWidget)
 {
 	USequencerPlaylistPlayer* Player = WeakPlayer.Get();
 	if (!ensure(Player))
@@ -457,7 +645,7 @@ FReply SSequencerPlaylistPanel::OnClicked_Item_Play(TSharedPtr<SSequencerPlaylis
 }
 
 
-FReply SSequencerPlaylistPanel::OnClicked_Item_Stop(TSharedPtr<SSequencerPlaylistItemWidget> ItemWidget)
+FReply SSequencerPlaylistPanel::HandleClicked_Item_Stop(TSharedPtr<SSequencerPlaylistItemWidget> ItemWidget)
 {
 	USequencerPlaylistPlayer* Player = WeakPlayer.Get();
 	if (!ensure(Player))
@@ -470,7 +658,7 @@ FReply SSequencerPlaylistPanel::OnClicked_Item_Stop(TSharedPtr<SSequencerPlaylis
 }
 
 
-FReply SSequencerPlaylistPanel::OnClicked_Item_Reset(TSharedPtr<SSequencerPlaylistItemWidget> ItemWidget)
+FReply SSequencerPlaylistPanel::HandleClicked_Item_Reset(TSharedPtr<SSequencerPlaylistItemWidget> ItemWidget)
 {
 	USequencerPlaylistPlayer* Player = WeakPlayer.Get();
 	if (!ensure(Player))
@@ -483,7 +671,7 @@ FReply SSequencerPlaylistPanel::OnClicked_Item_Reset(TSharedPtr<SSequencerPlayli
 }
 
 
-FReply SSequencerPlaylistPanel::OnClicked_Item_Close(TSharedPtr<SSequencerPlaylistItemWidget> ItemWidget)
+FReply SSequencerPlaylistPanel::HandleClicked_Item_Remove(TSharedPtr<SSequencerPlaylistItemWidget> ItemWidget)
 {
 	USequencerPlaylist* Playlist = GetCheckedPlaylist();
 	if (!Playlist)
@@ -492,60 +680,52 @@ FReply SSequencerPlaylistPanel::OnClicked_Item_Close(TSharedPtr<SSequencerPlayli
 	}
 
 	ensure(Playlist->Items.RemoveSingle(ItemWidget->GetItem()));
-	RegenerateSequenceList();
+	RegenerateRows(); // TODO: Incremental view update? (UETOOL-4644)
 	return FReply::Handled();
 }
 
-USequencerPlaylist* SSequencerPlaylistPanel::GetCheckedPlaylist()
+
+bool SSequencerPlaylistPanel::HandleItemDetailsIsPropertyVisible(const FPropertyAndParent& PropertyAndParent)
 {
-	USequencerPlaylistPlayer* Player = WeakPlayer.Get();
-	if (!ensure(Player))
+	static TMap<FName, FName> PropertyNameToColumnName;
+	static bool bMapInitialized = false;
+	if (bMapInitialized == false)
 	{
-		return nullptr;
-	}
-
-	USequencerPlaylist* Playlist = Player->GetPlaylist();
-	ensure(Playlist);
-	return Playlist;
-}
+		PropertyNameToColumnName.Add("Sequence", ColumnName_Items);
 
+		PropertyNameToColumnName.Add("StartFrameOffset", ColumnName_Offset);
+		PropertyNameToColumnName.Add("EndFrameOffset", ColumnName_Offset);
+		PropertyNameToColumnName.Add("bHoldAtFirstFrame", ColumnName_Hold);
+		PropertyNameToColumnName.Add("NumLoops", ColumnName_Loop);
 
-FReply SSequencerPlaylistPanel::HandleDragDetected(const FGeometry& MyGeometry, const FPointerEvent& MouseEvent, int32 SlotIndex, SVerticalBox::FSlot* Slot)
-{
-	USequencerPlaylist* Playlist = GetCheckedPlaylist();
-	if (!Playlist)
-	{
-		return FReply::Unhandled();
+		bMapInitialized = true;
 	}
 
-	if (ensure(Playlist->Items.IsValidIndex(SlotIndex)))
+	if (const FName* ColumnName = PropertyNameToColumnName.Find(PropertyAndParent.Property.GetFName()))
 	{
-		TSharedPtr<SBox> DragWidget = SNew(SBox)
-			.MinDesiredWidth(MyGeometry.GetLocalSize().X)
-			[
-				SNew(SSequencerPlaylistItemWidget, Playlist->Items[SlotIndex])
-			];
-
-		return FReply::Handled().BeginDragDrop(FSequencerPlaylistItemDragDropOp::New(SlotIndex, Slot, DragWidget));
+		if (ItemListView->GetHeaderRow()->IsColumnVisible(*ColumnName))
+		{
+			return false;
+		}
 	}
 
-	return FReply::Unhandled();
+	return true;
 }
 
 
-TOptional<SDragAndDropVerticalBox::EItemDropZone> SSequencerPlaylistPanel::HandleCanAcceptDrop(const FDragDropEvent& DragDropEvent, SDragAndDropVerticalBox::EItemDropZone DropZone, SVerticalBox::FSlot* Slot)
+TOptional<EItemDropZone> SSequencerPlaylistPanel::HandleCanAcceptDrop(const FDragDropEvent& DragDropEvent, EItemDropZone DropZone, TSharedPtr<FSequencerPlaylistRowData> RowData)
 {
 	TSharedPtr<FSequencerPlaylistItemDragDropOp> DragDropOperation = DragDropEvent.GetOperationAs<FSequencerPlaylistItemDragDropOp>();
 	if (DragDropOperation.IsValid())
 	{
-		return DropZone;
+		return DropZone == EItemDropZone::OntoItem ? EItemDropZone::BelowItem : DropZone;
 	}
 
-	return TOptional<SDragAndDropVerticalBox::EItemDropZone>();
+	return TOptional<EItemDropZone>();
 }
 
 
-FReply SSequencerPlaylistPanel::HandleAcceptDrop(const FDragDropEvent& DragDropEvent, SDragAndDropVerticalBox::EItemDropZone DropZone, int32 SlotIndex, SVerticalBox::FSlot* Slot)
+FReply SSequencerPlaylistPanel::HandleAcceptDrop(const FDragDropEvent& DragDropEvent, EItemDropZone DropZone, TSharedPtr<FSequencerPlaylistRowData> RowData)
 {
 	USequencerPlaylist* Playlist = GetCheckedPlaylist();
 	if (!Playlist)
@@ -556,15 +736,40 @@ FReply SSequencerPlaylistPanel::HandleAcceptDrop(const FDragDropEvent& DragDropE
 	TSharedPtr<FSequencerPlaylistItemDragDropOp> Operation = DragDropEvent.GetOperationAs<FSequencerPlaylistItemDragDropOp>();
 	if (Operation.IsValid())
 	{
-		const int32 SrcIndex = Operation->SlotIndexBeingDragged;
-		const int32 DstIndex = SlotIndex;
-		TArray<TObjectPtr<USequencerPlaylistItem>>& Items = Playlist->Items;
-		if (ensure(Items.IsValidIndex(SrcIndex)) && ensure(Items.IsValidIndex(DstIndex)))
+		// TODO: Not currently handling (potentially disjoint) multi-select,
+		// in part because there's no Algo::StablePartition.
+		// The ListView is set to ESelectionMode::Single for the time being.
+		if (ensure(Operation->SelectedItems.Num() == 1))
 		{
-			USequencerPlaylistItem* ItemToMove = Items[SrcIndex];
-			Items.RemoveAt(SrcIndex);
-			Items.Insert(ItemToMove, DstIndex);
-			return FReply::Handled();
+			const int32 SrcIndex = Operation->SelectedItems[0]->PlaylistIndex;
+			const int32 DropTargetIndex = RowData->PlaylistIndex;
+
+			if (SrcIndex == DropTargetIndex)
+			{
+				return FReply::Handled();
+			}
+
+			int32 DestIndexAdjustment = 0;
+			if (DropZone == EItemDropZone::BelowItem)
+			{
+				DestIndexAdjustment += 1;
+			}
+
+			if (SrcIndex < DropTargetIndex)
+			{
+				DestIndexAdjustment -= 1;
+			}
+
+			const int32 DestIndex = DropTargetIndex + DestIndexAdjustment;
+			TArray<TObjectPtr<USequencerPlaylistItem>>& Items = Playlist->Items;
+			if (ensure(Items.IsValidIndex(SrcIndex)) && ensure(Items.IsValidIndex(DestIndex)))
+			{
+				USequencerPlaylistItem* ItemToMove = Items[SrcIndex];
+				Items.RemoveAt(SrcIndex);
+				Items.Insert(ItemToMove, DestIndex);
+				RegenerateRows(); // TODO: Incremental view update? (UETOOL-4644)
+				return FReply::Handled();
+			}
 		}
 	}
 
@@ -572,14 +777,26 @@ FReply SSequencerPlaylistPanel::HandleAcceptDrop(const FDragDropEvent& DragDropE
 }
 
 
-TSharedRef<FSequencerPlaylistItemDragDropOp> FSequencerPlaylistItemDragDropOp::New(int32 InSlotIndexBeingDragged, SVerticalBox::FSlot* InSlotBeingDragged, TSharedPtr<SWidget> WidgetToShow)
+TSharedRef<FSequencerPlaylistItemDragDropOp> FSequencerPlaylistItemDragDropOp::New(const TArray<TSharedPtr<FSequencerPlaylistRowData>>& InSelectedItems)
 {
 	TSharedRef<FSequencerPlaylistItemDragDropOp> Operation = MakeShared<FSequencerPlaylistItemDragDropOp>();
 
+	Operation->SelectedItems = InSelectedItems;
+
 	Operation->MouseCursor = EMouseCursor::GrabHandClosed;
-	Operation->SlotIndexBeingDragged = InSlotIndexBeingDragged;
-	Operation->SlotBeingDragged = InSlotBeingDragged;
-	Operation->WidgetToShow = WidgetToShow;
+	Operation->Decorator = SNew(SBorder)
+		.BorderImage(FEditorStyle::GetBrush("Graph.ConnectorFeedback.Border"))
+		.Content()
+		[
+			SNew(SHorizontalBox)
+			+ SHorizontalBox::Slot()
+			.AutoWidth()
+			.VAlign(VAlign_Center)
+			[
+				SNew(STextBlock)
+				.Text(FText::Format(LOCTEXT("ItemDragHint", "Move {0} {0}|plural(one=item,other=items)"), InSelectedItems.Num()))
+			]
+		];
 
 	Operation->Construct();
 
@@ -592,167 +809,246 @@ FSequencerPlaylistItemDragDropOp::~FSequencerPlaylistItemDragDropOp()
 }
 
 
-TSharedPtr<SWidget> FSequencerPlaylistItemDragDropOp::GetDefaultDecorator() const
+void SSequencerPlaylistItemWidget::Construct(const FArguments& InArgs, TSharedPtr<FSequencerPlaylistRowData> InRowData, const TSharedRef<STableViewBase>& OwnerTableView)
 {
-	return SNew(SBorder)
-		.BorderImage(FEditorStyle::GetBrush("ContentBrowser.AssetDragDropTooltipBackground"))
-		.Content()
-		[
-			WidgetToShow.ToSharedRef()
-		];
+	check(InRowData && InRowData->WeakItem.IsValid());
+	RowData = InRowData;
+
+	PlayClickedDelegate = InArgs._OnPlayClicked;
+	StopClickedDelegate = InArgs._OnStopClicked;
+	ResetClickedDelegate = InArgs._OnResetClicked;
+	RemoveClickedDelegate = InArgs._OnRemoveClicked;
+	IsPropertyVisibleDelegate = InArgs._OnIsPropertyVisible;
+
+	FSuperRowType::Construct(
+		FSuperRowType::FArguments()
+			.OnCanAcceptDrop(InArgs._OnCanAcceptDrop)
+			.OnAcceptDrop(InArgs._OnAcceptDrop)
+			.OnDragDetected(this, &SSequencerPlaylistItemWidget::HandleDragDetected)
+		, OwnerTableView);
 }
 
 
 BEGIN_SLATE_FUNCTION_BUILD_OPTIMIZATION
-void SSequencerPlaylistItemWidget::Construct(const FArguments& InArgs, USequencerPlaylistItem* InItem)
+TSharedRef<SWidget> SSequencerPlaylistItemWidget::GenerateWidgetForColumn(const FName& ColumnName)
 {
-	check(InItem);
-	WeakItem = InItem;
+	USequencerPlaylistItem* Item = GetItem();
 
-	USequencerPlaylistItem_Sequence* Item =
-		CastChecked<USequencerPlaylistItem_Sequence>(InItem);
-
-	const ISlateStyle& SpmStyle = FSequencerPlaylistsStyle::Get();
+	if (!ensure(Item))
+	{
+		return SNullWidget::NullWidget;
+	}
 
 	static const FName PropertyEditorModuleName("PropertyEditor");
 	FPropertyEditorModule& PropertyEditorModule = FModuleManager::GetModuleChecked<FPropertyEditorModule>(PropertyEditorModuleName);
+	FSinglePropertyParams SinglePropParams;
+	SinglePropParams.NamePlacement = EPropertyNamePlacement::Hidden;
 
-	FSinglePropertyParams SequenceViewParams;
-	SequenceViewParams.bHideAssetThumbnail = true;
-	TSharedPtr<ISinglePropertyView> PropViewSequence = PropertyEditorModule.CreateSingleProperty(
-		Item, GET_MEMBER_NAME_CHECKED(USequencerPlaylistItem_Sequence, Sequence), SequenceViewParams);
+	if (ColumnName == SSequencerPlaylistPanel::ColumnName_Items)
+	{
+		USequencerPlaylistItem_Sequence* SequenceItem =
+			CastChecked<USequencerPlaylistItem_Sequence>(Item);
 
-	TSharedPtr<ISinglePropertyView> PropViewStartFrameOffset = PropertyEditorModule.CreateSingleProperty(
-		Item, GET_MEMBER_NAME_CHECKED(USequencerPlaylistItem, StartFrameOffset), {});
+		TSharedPtr<ISinglePropertyView> PropView = PropertyEditorModule.CreateSingleProperty(SequenceItem,
+			GET_MEMBER_NAME_CHECKED(USequencerPlaylistItem_Sequence, Sequence),
+			SinglePropParams);
 
-	TSharedPtr<ISinglePropertyView> PropViewEndFrameOffset = PropertyEditorModule.CreateSingleProperty(
-		Item, GET_MEMBER_NAME_CHECKED(USequencerPlaylistItem, EndFrameOffset), {});
+		return SNew(SBox)
+			.Padding(FMargin(8.0f, 0.0f, 0.0f, 0.0f))
+			[
+				PropView.ToSharedRef()
+			];
+	}
 
-	TSharedPtr<ISinglePropertyView> PropViewHold = PropertyEditorModule.CreateSingleProperty(
-		Item, GET_MEMBER_NAME_CHECKED(USequencerPlaylistItem, bHoldAtFirstFrame), {});
+	if (ColumnName == SSequencerPlaylistPanel::ColumnName_Offset)
+	{
+		TSharedPtr<ISinglePropertyView> StartOffsetPropView = PropertyEditorModule.CreateSingleProperty(
+			Item, GET_MEMBER_NAME_CHECKED(USequencerPlaylistItem, StartFrameOffset), SinglePropParams);
 
-	TSharedPtr<ISinglePropertyView> PropViewNumLoops = PropertyEditorModule.CreateSingleProperty(
-		Item, GET_MEMBER_NAME_CHECKED(USequencerPlaylistItem, NumLoops), {});
+		TSharedPtr<ISinglePropertyView> EndOffsetPropView = PropertyEditorModule.CreateSingleProperty(
+			Item, GET_MEMBER_NAME_CHECKED(USequencerPlaylistItem, EndFrameOffset), SinglePropParams);
 
-	ChildSlot
-	.Padding(10.0f)
-	[
-		SNew(SBorder)
-		.BorderImage(FEditorStyle::GetBrush("ToolPanel.GroupBorder"))
-		.Padding(FMargin(4.0f, 4.0f))
-		[
-			SNew(SHorizontalBox)
-			+ SHorizontalBox::Slot()
-			.FillWidth(0.02f)
-			.Padding(0.0f, 0.0f, 4.0f, 0.0f)
+		return SNew(SVerticalBox)
+			+ SVerticalBox::Slot()
 			[
-				SNew(SBorder)
-				.BorderImage(SpmStyle.GetBrush("SequencerPlaylists.SequenceBorderBrush"))
-				.VAlign(VAlign_Center)
-				.HAlign(HAlign_Center)
-				[
-					SNew(SImage)
-					.Image(FCoreStyle::Get().GetBrush("Icons.DragHandle"))
-				]
+				StartOffsetPropView.ToSharedRef()
 			]
-			+ SHorizontalBox::Slot()
-			.FillWidth(1.0f)
+			+ SVerticalBox::Slot()
 			[
-				SNew(SVerticalBox)
-				+ SVerticalBox::Slot()
-				.HAlign(HAlign_Fill)
-				[
-					SNew(SHorizontalBox)
-					+ SHorizontalBox::Slot()
-					.FillWidth(1.0f)
-					[
-						PropViewSequence.ToSharedRef()
-					]
-					+ SHorizontalBox::Slot()
-					.FillWidth(0.5f)
-					[
-						PropViewStartFrameOffset.ToSharedRef()
-					]
-					+ SHorizontalBox::Slot()
-					.FillWidth(1.0f)
-					[
-						SNew(SButton)
-						.Text(LOCTEXT("PlaySequenceButton", "Play"))
-						.HAlign(HAlign_Center)
-						.OnClicked_Lambda([this, InArgs]()
-							{
-								return InArgs._OnPlayClicked.IsBound() ? InArgs._OnPlayClicked.Execute(SharedThis(this)) : FReply::Unhandled();
-							})
-					]
-				]
-				+ SVerticalBox::Slot()
-				.HAlign(HAlign_Fill)
+				EndOffsetPropView.ToSharedRef()
+			];
+	}
+
+	if (ColumnName == SSequencerPlaylistPanel::ColumnName_Hold)
+	{
+		TSharedPtr<ISinglePropertyView> PropView = PropertyEditorModule.CreateSingleProperty(Item,
+			GET_MEMBER_NAME_CHECKED(USequencerPlaylistItem, bHoldAtFirstFrame),
+			SinglePropParams);
+
+		return PropView.ToSharedRef();
+	}
+
+	if (ColumnName == SSequencerPlaylistPanel::ColumnName_Loop)
+	{
+		TSharedPtr<ISinglePropertyView> PropView = PropertyEditorModule.CreateSingleProperty(Item,
+			GET_MEMBER_NAME_CHECKED(USequencerPlaylistItem, NumLoops),
+			SinglePropParams);
+
+		return PropView.ToSharedRef();
+	}
+
+	if (ColumnName == SSequencerPlaylistPanel::ColumnName_HoverDetails)
+	{
+		return SAssignNew(DetailsAnchor, SMenuAnchor)
+			.Placement(MenuPlacement_MenuLeft)
+			.OnGetMenuContent(this, &SSequencerPlaylistItemWidget::EnsureSelectedAndBuildContextMenu)
+			[
+				SNew(SBox)
+				.HeightOverride(24.0f)
+				.WidthOverride(10.0f)
 				[
-					SNew(SHorizontalBox)
-					+ SHorizontalBox::Slot()
-					.FillWidth(0.5f)
-					[
-						PropViewHold.ToSharedRef()
-					]
-					+ SHorizontalBox::Slot()
-					.FillWidth(0.5f)
-					[
-						PropViewNumLoops.ToSharedRef()
-					]
-					+ SHorizontalBox::Slot()
-					.FillWidth(0.5f)
-					[
-						PropViewEndFrameOffset.ToSharedRef()
-					]
-					+ SHorizontalBox::Slot()
-					.FillWidth(1.0f)
+					SNew(SButton)
+					.ButtonStyle(FEditorStyle::Get(), "HoverHintOnly")
+					.ContentPadding(0)
+					.OnClicked_Lambda([this]() { DetailsAnchor->SetIsOpen(!DetailsAnchor->IsOpen()); return FReply::Handled(); })
+					.VAlign(VAlign_Center)
+					.HAlign(HAlign_Center)
+					.Visibility_Lambda([this]() { return this->IsHovered() ? EVisibility::Visible : EVisibility::Hidden; })
 					[
-						SNew(SHorizontalBox)
-						+ SHorizontalBox::Slot()
-						.FillWidth(1.0f)
-						[
-							SNew(SButton)
-							.Text(LOCTEXT("StopSequenceButton", "Stop"))
-							.HAlign(HAlign_Center)
-							.OnClicked_Lambda([this, InArgs]()
-								{
-									return InArgs._OnStopClicked.IsBound() ? InArgs._OnStopClicked.Execute(SharedThis(this)) : FReply::Unhandled();
-								})
-						]
-						+ SHorizontalBox::Slot()
-						.FillWidth(1.0f)
-						[
-							SNew(SButton)
-							.Text(LOCTEXT("ResetSequenceButton", "Reset"))
-							.HAlign(HAlign_Center)
-							.OnClicked_Lambda([this, InArgs]()
-								{
-									return InArgs._OnResetClicked.IsBound() ? InArgs._OnResetClicked.Execute(SharedThis(this)) : FReply::Unhandled();
-								})
-						]
+						SNew(SImage)
+						.ColorAndOpacity(FSlateColor::UseForeground())
+						.Image(FSequencerPlaylistsStyle::Get().GetBrush("SequencerPlaylists.Ellipsis"))
 					]
 				]
-			]
-			+ SHorizontalBox::Slot()
-			.AutoWidth()
-			.VAlign(VAlign_Top)
-			[
-				SNew(SButton)
-				.ButtonStyle(FAppStyle::Get(), "SimpleButton")
-				.OnClicked_Lambda([this, InArgs]()
-					{
-						return InArgs._OnCloseClicked.IsBound() ? InArgs._OnCloseClicked.Execute(SharedThis(this)) : FReply::Unhandled();
-					})
-				[
-					SNew(SImage)
-					.Image(FAppStyle::Get().GetBrush("Icons.X"))
-					.ColorAndOpacity(FSlateColor::UseForeground())
-				]
-			]
-		]
-	];
+			];
+	}
+
+	return SNullWidget::NullWidget;
 }
 END_SLATE_FUNCTION_BUILD_OPTIMIZATION
 
 
+FReply SSequencerPlaylistItemWidget::OnMouseButtonUp(const FGeometry& MyGeometry, const FPointerEvent& MouseEvent)
+{
+	if (MouseEvent.GetEffectingButton() == EKeys::RightMouseButton)
+	{
+		TSharedRef<SWidget> MenuWidget = EnsureSelectedAndBuildContextMenu();
+		const FWidgetPath WidgetPath = MouseEvent.GetEventPath() ? *MouseEvent.GetEventPath() : FWidgetPath();
+		FSlateApplication::Get().PushMenu(SharedThis(this), WidgetPath, MenuWidget, MouseEvent.GetScreenSpacePosition(), FPopupTransitionEffect(FPopupTransitionEffect::ContextMenu));
+		return FReply::Handled().ReleaseMouseCapture().SetUserFocus(MenuWidget, EFocusCause::SetDirectly);
+	}
+
+	return FSuperRowType::OnMouseButtonUp(MyGeometry, MouseEvent);
+}
+
+
+FReply SSequencerPlaylistItemWidget::HandleDragDetected(const FGeometry& MyGeometry, const FPointerEvent& MouseEvent)
+{
+	TSharedPtr<ITypedTableView<TSharedPtr<FSequencerPlaylistRowData>>> OwnerTable = OwnerTablePtr.Pin();
+	if (OwnerTable)
+	{
+		return FReply::Handled().BeginDragDrop(FSequencerPlaylistItemDragDropOp::New(OwnerTable->GetSelectedItems()));
+	}
+
+	return FReply::Unhandled();
+}
+
+
+TSharedRef<SWidget> SSequencerPlaylistItemWidget::EnsureSelectedAndBuildContextMenu()
+{
+	TArray<UObject*> SelectedItems;
+
+	TSharedPtr<ITypedTableView<TSharedPtr<FSequencerPlaylistRowData>>> OwnerTable = OwnerTablePtr.Pin();
+	if (OwnerTable)
+	{
+		TArray<TSharedPtr<FSequencerPlaylistRowData>> SelectedRows = OwnerTable->GetSelectedItems();
+		SelectedItems.Reserve(SelectedRows.Num());
+		for (const TSharedPtr<FSequencerPlaylistRowData>& SelectedRow : SelectedRows)
+		{
+			SelectedItems.Add(SelectedRow->WeakItem.Get());
+		}
+
+		if (!SelectedItems.Contains(GetItem()))
+		{
+			OwnerTable->Private_ClearSelection();
+			OwnerTable->Private_SetItemSelection(RowData, true, true);
+			OwnerTable->Private_SignalSelectionChanged(ESelectInfo::OnMouseClick);
+		}
+	}
+	else
+	{
+		SelectedItems.Add(GetItem());
+	}
+
+	return BuildContextMenu(SelectedItems);
+}
+
+
+TSharedRef<SWidget> SSequencerPlaylistItemWidget::BuildContextMenu(const TArray<UObject*>& SelectedItems)
+{
+	TSharedRef<SSequencerPlaylistItemWidget> This = SharedThis(this);
+
+	FSequencerPlaylistsModule& Module = static_cast<FSequencerPlaylistsModule&>(ISequencerPlaylistsModule::Get());
+	FMenuBuilder MenuBuilder(true, Module.GetCommandList());
+
+	MenuBuilder.BeginSection(NAME_None, LOCTEXT("ItemContextPlaybackHeading", "Playback"));
+	{
+		MenuBuilder.AddMenuEntry(
+			LOCTEXT("ItemContextPlay", "Play"),
+			LOCTEXT("ItemContextPlayTooltip", "Trigger playback of item"),
+			FSlateIcon(),
+			FUIAction(FExecuteAction::CreateLambda([This]() { This->PlayClickedDelegate.Execute(This); }))
+		);
+
+		MenuBuilder.AddMenuEntry(
+			LOCTEXT("ItemContextStop", "Stop"),
+			LOCTEXT("ItemContextStopTooltip", "Stop playback of this item"),
+			FSlateIcon(),
+			FUIAction(FExecuteAction::CreateLambda([This]() { This->StopClickedDelegate.Execute(This); }))
+		);
+
+		MenuBuilder.AddMenuEntry(
+			LOCTEXT("ItemContextReset", "Reset"),
+			LOCTEXT("ItemContextResetTooltip", "Reset playback of this item"),
+			FSlateIcon(),
+			FUIAction(FExecuteAction::CreateLambda([This]() { This->ResetClickedDelegate.Execute(This); }))
+		);
+	}
+
+	MenuBuilder.BeginSection(NAME_None, LOCTEXT("ItemContextDetailsHeading", "Details"));
+	{
+		FDetailsViewArgs DetailsViewArgs;
+		{
+			DetailsViewArgs.bAllowSearch = false;
+			DetailsViewArgs.bCustomFilterAreaLocation = true;
+			DetailsViewArgs.bCustomNameAreaLocation = true;
+			DetailsViewArgs.bHideSelectionTip = true;
+			DetailsViewArgs.bLockable = false;
+			DetailsViewArgs.bSearchInitialKeyFocus = true;
+			DetailsViewArgs.bUpdatesFromSelection = false;
+			DetailsViewArgs.bShowOptions = false;
+			DetailsViewArgs.bShowModifiedPropertiesOption = false;
+			DetailsViewArgs.ColumnWidth = 0.45f;
+		}
+
+		TSharedRef<IDetailsView> DetailsView = FModuleManager::GetModuleChecked<FPropertyEditorModule>("PropertyEditor").CreateDetailView(DetailsViewArgs);
+		DetailsView->SetIsPropertyVisibleDelegate(IsPropertyVisibleDelegate);
+		DetailsView->SetObjects(SelectedItems);
+		MenuBuilder.AddWidget(DetailsView, FText::GetEmpty(), true);
+	}
+
+	MenuBuilder.BeginSection(NAME_None, LOCTEXT("ItemContextEditHeading", "Edit"));
+	{
+		MenuBuilder.AddMenuEntry(
+			LOCTEXT("ItemContextRemove", "Remove from Playlist"),
+			LOCTEXT("ItemContextRemoveTooltip", "Remove this item from the Playlist"),
+			FSlateIcon(FEditorStyle::GetStyleSetName(), "ContentBrowser.AssetActions.Delete"),
+			FUIAction(FExecuteAction::CreateLambda([This]() { This->RemoveClickedDelegate.Execute(This); }))
+		);
+	}
+
+	return MenuBuilder.MakeWidget();
+}
+
+
 #undef LOCTEXT_NAMESPACE
diff --git a/Engine/Plugins/VirtualProduction/SequencerPlaylists/Source/SequencerPlaylists/Private/SequencerPlaylistsWidgets.h b/Engine/Plugins/VirtualProduction/SequencerPlaylists/Source/SequencerPlaylists/Private/SequencerPlaylistsWidgets.h
--- a/Engine/Plugins/VirtualProduction/SequencerPlaylists/Source/SequencerPlaylists/Private/SequencerPlaylistsWidgets.h
+++ b/Engine/Plugins/VirtualProduction/SequencerPlaylists/Source/SequencerPlaylists/Private/SequencerPlaylistsWidgets.h
@@ -3,18 +3,37 @@
 #pragma once
 
 #include "CoreMinimal.h"
+#include "PropertyEditorDelegates.h"
 #include "Widgets/DeclarativeSyntaxSupport.h"
-#include "Widgets/SBoxPanel.h"
 #include "Widgets/SCompoundWidget.h"
+#include "Widgets/Views/SListView.h"
 
 
 struct FAssetData;
+class SMenuAnchor;
+class SMultiLineEditableTextBox;
 class SSequencerPlaylistItemWidget;
+class SSearchBox;
+class STextBlock;
+template <typename ItemType> class TTextFilter;
 class USequencerPlaylist;
 class USequencerPlaylistItem;
 class USequencerPlaylistPlayer;
 
 
+struct FSequencerPlaylistRowData
+{
+	int32 PlaylistIndex;
+	TWeakObjectPtr<USequencerPlaylistItem> WeakItem;
+
+	FSequencerPlaylistRowData(int32 InPlaylistIndex, USequencerPlaylistItem* InItem)
+		: PlaylistIndex(InPlaylistIndex)
+		, WeakItem(InItem)
+	{
+	}
+};
+
+
 class SSequencerPlaylistPanel : public SCompoundWidget
 {
 	SLATE_BEGIN_ARGS(SSequencerPlaylistPanel) {}
@@ -22,72 +41,123 @@ class SSequencerPlaylistPanel : public SCompoundWidget
 
 public:
 	static const float DefaultWidth;
+	static const FName ColumnName_HoverTransport;
+	static const FName ColumnName_Items;
+	static const FName ColumnName_Offset;
+	static const FName ColumnName_Hold;
+	static const FName ColumnName_Loop;
+	static const FName ColumnName_HoverDetails;
 
 	void Construct(const FArguments& InArgs, USequencerPlaylistPlayer* InPlayer);
 
 private:
-	void RegenerateSequenceList();
+	TSharedRef<SWidget> Construct_Toolbar();
+	TSharedRef<SWidget> Construct_Transport();
+	TSharedRef<SWidget> Construct_AddSearchRow();
+	TSharedRef<SWidget> Construct_ItemListView();
 
-	TSharedRef<SWidget> OnPresetGeneratePresetsMenu();
-	void OnSaveAsPreset();
-	void OnLoadPreset(const FAssetData& InPreset);
+	USequencerPlaylist* GetCheckedPlaylist();
+	void RegenerateRows();
 
-	FReply OnClicked_PlayAll();
-	FReply OnClicked_StopAll();
-	FReply OnClicked_ResetAll();
-	FReply OnClicked_AddSequence();
+	TSharedRef<SWidget> BuildOpenPlaylistMenu();
+	void OnSavePlaylistAs();
+	void OnLoadPlaylist(const FAssetData& InPreset);
+	void OnNewPlaylist();
 
-	FReply OnClicked_Item_Play(TSharedPtr<SSequencerPlaylistItemWidget> ItemWidget);
-	FReply OnClicked_Item_Stop(TSharedPtr<SSequencerPlaylistItemWidget> ItemWidget);
-	FReply OnClicked_Item_Reset(TSharedPtr<SSequencerPlaylistItemWidget> ItemWidget);
-	FReply OnClicked_Item_Close(TSharedPtr<SSequencerPlaylistItemWidget> ItemWidget);
+	void GetSearchStrings(const FSequencerPlaylistRowData& Item, TArray<FString>& OutSearchStrings);
+	void OnSearchTextChanged(const FText& InFilterText);
 
-	USequencerPlaylist* GetCheckedPlaylist();
+	FReply HandleClicked_PlayAll();
+	FReply HandleClicked_StopAll();
+	FReply HandleClicked_ResetAll();
+	FReply HandleClicked_AddSequence();
+
+	FReply HandleClicked_Item_Play(TSharedPtr<SSequencerPlaylistItemWidget> ItemWidget);
+	FReply HandleClicked_Item_Stop(TSharedPtr<SSequencerPlaylistItemWidget> ItemWidget);
+	FReply HandleClicked_Item_Reset(TSharedPtr<SSequencerPlaylistItemWidget> ItemWidget);
+	FReply HandleClicked_Item_Remove(TSharedPtr<SSequencerPlaylistItemWidget> ItemWidget);
+
+	bool HandleItemDetailsIsPropertyVisible(const FPropertyAndParent& PropertyAndParent);
 
 	// Drag and drop handling
-	FReply HandleDragDetected(const FGeometry& MyGeometry, const FPointerEvent& MouseEvent, int32 SlotIndex, SVerticalBox::FSlot* Slot);
-	TOptional<SDragAndDropVerticalBox::EItemDropZone> HandleCanAcceptDrop(const FDragDropEvent& DragDropEvent, SDragAndDropVerticalBox::EItemDropZone DropZone, SVerticalBox::FSlot* Slot);
-	FReply HandleAcceptDrop(const FDragDropEvent& DragDropEvent, SDragAndDropVerticalBox::EItemDropZone DropZone, int32 SlotIndex, SVerticalBox::FSlot* Slot);
+	TOptional<EItemDropZone> HandleCanAcceptDrop(const FDragDropEvent& DragDropEvent, EItemDropZone DropZone, TSharedPtr<FSequencerPlaylistRowData> RowData);
+	FReply HandleAcceptDrop(const FDragDropEvent& DragDropEvent, EItemDropZone DropZone, TSharedPtr<FSequencerPlaylistRowData> RowData);
 
 private:
 	TWeakObjectPtr<USequencerPlaylistPlayer> WeakPlayer;
-	TSharedPtr<SDragAndDropVerticalBox> ItemList;
+
+	TSharedPtr<SSearchBox> SearchBox;
+	TSharedPtr<TTextFilter<const FSequencerPlaylistRowData&>> SearchTextFilter;
+	TArray<TSharedPtr<FSequencerPlaylistRowData>> ItemRows;
+	TSharedPtr<SListView<TSharedPtr<FSequencerPlaylistRowData>>> ItemListView;
 };
 
 
-class FSequencerPlaylistItemDragDropOp : public FDragAndDropVerticalBoxOp
+class FSequencerPlaylistItemDragDropOp : public FDragDropOperation
 {
 public:
-	DRAG_DROP_OPERATOR_TYPE(FSequencerPlaylistItemDragDropOp, FDragAndDropVerticalBoxOp)
+	DRAG_DROP_OPERATOR_TYPE(FSequencerPlaylistItemDragDropOp, FDragDropOperation)
 
-	TSharedPtr<SWidget> WidgetToShow;
+	TArray<TSharedPtr<FSequencerPlaylistRowData>> SelectedItems;
 
-	static TSharedRef<FSequencerPlaylistItemDragDropOp> New(int32 InSlotIndexBeingDragged, SVerticalBox::FSlot* InSlotBeingDragged, TSharedPtr<SWidget> InWidgetToShow);
+	static TSharedRef<FSequencerPlaylistItemDragDropOp> New(const TArray<TSharedPtr<FSequencerPlaylistRowData>>& InSelectedItems);
 
 public:
 	virtual ~FSequencerPlaylistItemDragDropOp();
 
-	virtual TSharedPtr<SWidget> GetDefaultDecorator() const override;
+	virtual TSharedPtr<SWidget> GetDefaultDecorator() const override
+	{
+		return Decorator;
+	}
+
+private:
+	TSharedPtr<SWidget> Decorator;
 };
 
 
 DECLARE_DELEGATE_RetVal_OneParam(FReply, FOnClickedSequencerPlaylistItem, TSharedPtr<SSequencerPlaylistItemWidget> /*ItemWidget*/);
 
 
-class SSequencerPlaylistItemWidget : public SCompoundWidget
+class SSequencerPlaylistItemWidget : public SMultiColumnTableRow<TSharedPtr<FSequencerPlaylistRowData>>
 {
 	SLATE_BEGIN_ARGS(SSequencerPlaylistItemWidget) {}
 		SLATE_EVENT(FOnClickedSequencerPlaylistItem, OnPlayClicked)
 		SLATE_EVENT(FOnClickedSequencerPlaylistItem, OnStopClicked)
 		SLATE_EVENT(FOnClickedSequencerPlaylistItem, OnResetClicked)
-		SLATE_EVENT(FOnClickedSequencerPlaylistItem, OnCloseClicked)
+		SLATE_EVENT(FOnClickedSequencerPlaylistItem, OnRemoveClicked)
+
+		SLATE_EVENT(FIsPropertyVisible, OnIsPropertyVisible)
+		SLATE_EVENT(FOnCanAcceptDrop, OnCanAcceptDrop)
+		SLATE_EVENT(FOnAcceptDrop, OnAcceptDrop)
 	SLATE_END_ARGS()
 
 public:
-	void Construct(const FArguments& InArgs, USequencerPlaylistItem* InItem);
+	void Construct(const FArguments& InArgs, TSharedPtr<FSequencerPlaylistRowData> InRowData, const TSharedRef<STableViewBase>& OwnerTableView);
+
+	const TSharedPtr<FSequencerPlaylistRowData>& GetRowData() { return RowData; }
+	USequencerPlaylistItem* GetItem() { return GetRowData() ? GetRowData()->WeakItem.Get() : nullptr; }
+
+	//~ Begin SMultiColumnTableRow
+	TSharedRef<SWidget> GenerateWidgetForColumn(const FName& InColumnName) override;
+	//~ End SMultiColumnTableRow
 
-	USequencerPlaylistItem* GetItem() { return WeakItem.Get(); }
+	//~ Begin SWidget
+	FReply OnMouseButtonUp(const FGeometry& MyGeometry, const FPointerEvent& MouseEvent) override;
+	//~ End SWidget
 
 private:
-	TWeakObjectPtr<USequencerPlaylistItem> WeakItem;
-};
\ No newline at end of file
+	FReply HandleDragDetected(const FGeometry& MyGeometry, const FPointerEvent& MouseEvent);
+
+	TSharedRef<SWidget> EnsureSelectedAndBuildContextMenu();
+	TSharedRef<SWidget> BuildContextMenu(const TArray<UObject*>& SelectedItems);
+
+private:
+	TSharedPtr<FSequencerPlaylistRowData> RowData;
+	TSharedPtr<SMenuAnchor> DetailsAnchor;
+
+	FOnClickedSequencerPlaylistItem PlayClickedDelegate;
+	FOnClickedSequencerPlaylistItem StopClickedDelegate;
+	FOnClickedSequencerPlaylistItem ResetClickedDelegate;
+	FOnClickedSequencerPlaylistItem RemoveClickedDelegate;
+	FIsPropertyVisible IsPropertyVisibleDelegate;
+};
diff --git a/Engine/Plugins/VirtualProduction/SequencerPlaylists/Source/SequencerPlaylists/Public/SequencerPlaylistItem.h b/Engine/Plugins/VirtualProduction/SequencerPlaylists/Source/SequencerPlaylists/Public/SequencerPlaylistItem.h
--- a/Engine/Plugins/VirtualProduction/SequencerPlaylists/Source/SequencerPlaylists/Public/SequencerPlaylistItem.h
+++ b/Engine/Plugins/VirtualProduction/SequencerPlaylists/Source/SequencerPlaylists/Public/SequencerPlaylistItem.h
@@ -19,20 +19,23 @@ public:
 	virtual FText GetDisplayName() PURE_VIRTUAL(USequencerPlaylistItem::GetDisplayName, return FText::GetEmpty(); )
 
 public:
-	UPROPERTY(EditAnywhere, Category="SequencerPlaylists")
+	UPROPERTY(EditAnywhere, Category="SequencerPlaylists", meta=(NoResetToDefault))
 	int32 StartFrameOffset;
 
-	UPROPERTY(EditAnywhere, Category="SequencerPlaylists")
+	UPROPERTY(EditAnywhere, Category="SequencerPlaylists", meta=(NoResetToDefault))
 	int32 EndFrameOffset;
 
 	/**
 	 * If true, the sequence will be inserted immediately on recording start and any time Reset()
 	 * is called, paused at the first frame indefinitely until either triggered or stopped.
 	 */
-	UPROPERTY(EditAnywhere, Category="SequencerPlaylists")
+	UPROPERTY(EditAnywhere, Category="SequencerPlaylists", meta=(NoResetToDefault))
 	bool bHoldAtFirstFrame;
 
 	/** 0 is single playthrough, >= 1 is (n+1) playthroughs. */
-	UPROPERTY(EditAnywhere, Category="SequencerPlaylists")
+	UPROPERTY(EditAnywhere, Category="SequencerPlaylists", meta=(NoResetToDefault))
 	int32 NumLoops;
+
+	UPROPERTY(EditAnywhere, Category="SequencerPlaylists", meta=(NoResetToDefault))
+	float PlaybackSpeed = 1.0f;
 };
diff --git a/Engine/Plugins/VirtualProduction/SequencerPlaylists/Source/SequencerPlaylists/SequencerPlaylists.Build.cs b/Engine/Plugins/VirtualProduction/SequencerPlaylists/Source/SequencerPlaylists/SequencerPlaylists.Build.cs
--- a/Engine/Plugins/VirtualProduction/SequencerPlaylists/Source/SequencerPlaylists/SequencerPlaylists.Build.cs
+++ b/Engine/Plugins/VirtualProduction/SequencerPlaylists/Source/SequencerPlaylists/SequencerPlaylists.Build.cs
@@ -7,7 +7,7 @@ public class SequencerPlaylists : ModuleRules
 	public SequencerPlaylists(ReadOnlyTargetRules Target) : base(Target)
 	{
 		PCHUsage = ModuleRules.PCHUsageMode.UseExplicitOrSharedPCHs;
-		
+
 		PublicIncludePaths.AddRange(new string[] {
 		});
 
@@ -36,6 +36,7 @@ public class SequencerPlaylists : ModuleRules
 			"MovieScene",
 			"LevelSequenceEditor",
 			"TakesCore",
+			"ToolWidgets",
 		});
 
 		DynamicallyLoadedModuleNames.AddRange(new string[] {
