Sequencer Playlists: Add per-item inline transport buttons, and support for Trigger Mode. Also fixes a bug with context menu activation on unselected items.

diff --git a/Engine/Plugins/VirtualProduction/SequencerPlaylists/Source/SequencerPlaylists/Private/SequencerPlaylistsStyle.cpp b/Engine/Plugins/VirtualProduction/SequencerPlaylists/Source/SequencerPlaylists/Private/SequencerPlaylistsStyle.cpp
--- a/Engine/Plugins/VirtualProduction/SequencerPlaylists/Source/SequencerPlaylists/Private/SequencerPlaylistsStyle.cpp
+++ b/Engine/Plugins/VirtualProduction/SequencerPlaylists/Source/SequencerPlaylists/Private/SequencerPlaylistsStyle.cpp
@@ -1,15 +1,18 @@
 // Copyright Epic Games, Inc. All Rights Reserved.
 
 #include "SequencerPlaylistsStyle.h"
-#include "Styling/SlateStyleRegistry.h"
+#include "EditorStyleSet.h"
 #include "Framework/Application/SlateApplication.h"
-#include "Slate/SlateGameResources.h"
 #include "Interfaces/IPluginManager.h"
+#include "Slate/SlateGameResources.h"
 #include "Styling/SlateStyleMacros.h"
+#include "Styling/SlateStyleRegistry.h"
 #include "Styling/StyleColors.h"
 
+
 #define RootToContentDir Style->RootToContentDir
 
+
 TSharedPtr<FSlateStyleSet> FSequencerPlaylistsStyle::StyleInstance = nullptr;
 
 
@@ -35,12 +38,25 @@ FName FSequencerPlaylistsStyle::GetStyleSetName()
 	return StyleSetName;
 }
 
+FLinearColor FSequencerPlaylistsStyle::MakeColorVariation(FLinearColor InColor, EColorVariation Variation)
+{
+	const FLinearColor HsvColor = InColor.LinearRGBToHSV();
+	switch (Variation)
+	{
+		case EColorVariation::Desaturated: return FLinearColor(HsvColor.R, HsvColor.G * .5, HsvColor.B, HsvColor.A).HSVToLinearRGB();
+		case EColorVariation::Dimmed: return FLinearColor(HsvColor.R, HsvColor.G, HsvColor.B * .5, HsvColor.A).HSVToLinearRGB();
+		default: checkNoEntry();
+	}
+
+	return InColor;
+};
+
 TSharedRef<FSlateStyleSet> FSequencerPlaylistsStyle::Create()
 {
+	static const FVector2D Icon12x12(12.0f, 12.0f);
 	static const FVector2D Icon16x16(16.0f, 16.0f);
 	static const FVector2D Icon20x20(20.0f, 20.0f);
 	static const FVector2D Icon24x24(24.0f, 24.0f);
-	static const FVector2D LargeTransport(60.0f, 60.0f);
 
 	TSharedRef<FSlateStyleSet> Style = MakeShareable(new FSlateStyleSet("SequencerPlaylistsStyle"));
 	Style->SetContentRoot(IPluginManager::Get().FindPlugin("SequencerPlaylists")->GetBaseDir() / TEXT("Resources"));
@@ -49,17 +65,24 @@ TSharedRef<FSlateStyleSet> FSequencerPlaylistsStyle::Create()
 	Style->Set("SequencerPlaylists.OpenPluginWindow", new IMAGE_BRUSH_SVG("PlaceholderButtonIcon", Icon20x20));
 	Style->Set("SequencerPlaylists.Panel.Background", new FSlateColorBrush(FStyleColors::Panel));
 
+	FLinearColor DimColor(FSlateColor(EStyleColor::Background).GetSpecifiedColor());
+	DimColor.A = 0.5f;
+	Style->Set("SequencerPlaylists.Item.Dim", new FSlateColorBrush(DimColor));
+
 	Style->Set("SequencerPlaylists.NewPlaylist.Background", new IMAGE_BRUSH_SVG("PlaylistIcon", Icon20x20));
 	// NewPlaylist.Overlay, SavePlaylistAs, OpenPlaylist are registered under a different root below.
 
-	const FButtonStyle& SimpleButtonStyle = Style->GetWidgetStyle<FButtonStyle>("SimpleButton");
-
-	FButtonStyle TransportButtonStyle = FButtonStyle(SimpleButtonStyle)
+	FButtonStyle TransportButtonStyle = FButtonStyle()
 		.SetNormalForeground(FStyleColors::Foreground)
-		.SetPressedForeground(FStyleColors::ForegroundHover)
-		.SetHoveredForeground(FStyleColors::ForegroundHover)
 		.SetDisabledForeground(FStyleColors::Foreground);
 
+	const FLinearColor PlayColor = FStyleColors::AccentGreen.GetSpecifiedColor();
+	const FLinearColor StopColor = FStyleColors::AccentRed.GetSpecifiedColor();
+	const FLinearColor ResetColor = FStyleColors::AccentBlue.GetSpecifiedColor();
+	const FLinearColor PlayPressedColor = MakeColorVariation(PlayColor, EColorVariation::Dimmed);
+	const FLinearColor StopPressedColor = MakeColorVariation(StopColor, EColorVariation::Dimmed);
+	const FLinearColor ResetPressedColor = MakeColorVariation(ResetColor, EColorVariation::Dimmed);
+
 	const float CornerRadius = 4.0f;
 	const FVector4 LeftCorners(CornerRadius, 0.0f, 0.0f, CornerRadius);
 	const FVector4 MiddleCorners(0.0f, 0.0f, 0.0f, 0.0f);
@@ -94,18 +117,45 @@ TSharedRef<FSlateStyleSet> FSequencerPlaylistsStyle::Create()
 		.SetPressedPadding(FMargin(6.f, 4.f, 8.f, 4.f));
 
 	const FButtonStyle PlayTransportButton = FButtonStyle(LeftTransportButton)
-		.SetHoveredForeground(FStyleColors::AccentGreen);
+		.SetHoveredForeground(PlayColor)
+		.SetPressedForeground(PlayPressedColor);
 
 	const FButtonStyle StopTransportButton = FButtonStyle(MiddleTransportButton)
-		.SetHoveredForeground(FStyleColors::AccentRed);
+		.SetHoveredForeground(StopColor)
+		.SetPressedForeground(StopPressedColor);
 
 	const FButtonStyle ResetTransportButton = FButtonStyle(RightTransportButton)
-		.SetHoveredForeground(FStyleColors::AccentBlue);
+		.SetHoveredForeground(ResetColor)
+		.SetPressedForeground(ResetPressedColor);
 
 	Style->Set("SequencerPlaylists.TransportButton.Play", PlayTransportButton);
 	Style->Set("SequencerPlaylists.TransportButton.Stop", StopTransportButton);
 	Style->Set("SequencerPlaylists.TransportButton.Reset", ResetTransportButton);
 
+	const FButtonStyle HoverTransportButtonStyle = FButtonStyle()
+		.SetNormal(FSlateNoResource())
+		.SetHovered(FSlateNoResource())
+		.SetPressed(FSlateNoResource())
+		.SetNormalForeground(FStyleColors::Foreground)
+		.SetNormalPadding(FMargin(0, 0, 0, 0))
+		.SetPressedPadding(FMargin(0, 0, 0, 0));
+
+	const FButtonStyle PlayHoverTransportButton = FButtonStyle(HoverTransportButtonStyle)
+		.SetHoveredForeground(PlayColor)
+		.SetPressedForeground(PlayPressedColor);
+
+	const FButtonStyle StopHoverTransportButton = FButtonStyle(HoverTransportButtonStyle)
+		.SetHoveredForeground(StopColor)
+		.SetPressedForeground(StopPressedColor);
+
+	const FButtonStyle ResetHoverTransportButton = FButtonStyle(HoverTransportButtonStyle)
+		.SetHoveredForeground(ResetColor)
+		.SetPressedForeground(ResetPressedColor);
+
+	Style->Set("SequencerPlaylists.HoverTransport.Play", PlayHoverTransportButton);
+	Style->Set("SequencerPlaylists.HoverTransport.Stop", StopHoverTransportButton);
+	Style->Set("SequencerPlaylists.HoverTransport.Reset", ResetHoverTransportButton);
+
 	FEditableTextBoxStyle EditableTextStyle = FEditableTextBoxStyle()
 		.SetFont(FCoreStyle::GetDefaultFontStyle("Regular", 9))
 		.SetBackgroundImageNormal(FSlateNoResource())
@@ -142,7 +192,9 @@ TSharedRef<FSlateStyleSet> FSequencerPlaylistsStyle::Create()
 			Style->Set("SequencerPlaylists.OpenPlaylist", new IMAGE_BRUSH_SVG("Starship/Common/OpenAsset", Icon16x16));
 
 			Style->Set("SequencerPlaylists.Play", new IMAGE_BRUSH_SVG("Starship/Common/play", Icon20x20));
+			Style->Set("SequencerPlaylists.Play.Small", new IMAGE_BRUSH_SVG("Starship/Common/play", Icon12x12));
 			Style->Set("SequencerPlaylists.Stop", new IMAGE_BRUSH_SVG("Starship/MainToolbar/stop", Icon20x20));
+			Style->Set("SequencerPlaylists.Stop.Small", new IMAGE_BRUSH_SVG("Starship/MainToolbar/stop", Icon12x12));
 		}
 	}
 
diff --git a/Engine/Plugins/VirtualProduction/SequencerPlaylists/Source/SequencerPlaylists/Private/SequencerPlaylistsStyle.h b/Engine/Plugins/VirtualProduction/SequencerPlaylists/Source/SequencerPlaylists/Private/SequencerPlaylistsStyle.h
--- a/Engine/Plugins/VirtualProduction/SequencerPlaylists/Source/SequencerPlaylists/Private/SequencerPlaylistsStyle.h
+++ b/Engine/Plugins/VirtualProduction/SequencerPlaylists/Source/SequencerPlaylists/Private/SequencerPlaylistsStyle.h
@@ -21,6 +21,14 @@ public:
 	static FName GetStyleSetName();
 
 private:
+	enum class EColorVariation
+	{
+		Desaturated,
+		Dimmed
+	};
+
+	static FLinearColor MakeColorVariation(FLinearColor InColor, EColorVariation Variation);
+
 	static TSharedRef< class FSlateStyleSet > Create();
 
 private:
diff --git a/Engine/Plugins/VirtualProduction/SequencerPlaylists/Source/SequencerPlaylists/Private/SequencerPlaylistsWidgets.cpp b/Engine/Plugins/VirtualProduction/SequencerPlaylists/Source/SequencerPlaylists/Private/SequencerPlaylistsWidgets.cpp
--- a/Engine/Plugins/VirtualProduction/SequencerPlaylists/Source/SequencerPlaylists/Private/SequencerPlaylistsWidgets.cpp
+++ b/Engine/Plugins/VirtualProduction/SequencerPlaylists/Source/SequencerPlaylists/Private/SequencerPlaylistsWidgets.cpp
@@ -72,8 +72,23 @@ void SSequencerPlaylistPanel::Construct(const FArguments& InArgs, USequencerPlay
 			SNew(SVerticalBox)
 			+SVerticalBox::Slot()
 			.AutoHeight()
+			.Padding(0.0f, 0.0f, 8.0f, 0.0f)
 			[
-				Construct_Toolbar()
+				SNew(SHorizontalBox)
+				+ SHorizontalBox::Slot()
+				.AutoWidth()
+				[
+					Construct_LeftToolbar()
+				]
+				+ SHorizontalBox::Slot()
+				[
+					SNew(SSpacer)
+				]
+				+ SHorizontalBox::Slot()
+				.AutoWidth()
+				[
+					Construct_RightToolbar()
+				]
 			]
 			+ SVerticalBox::Slot()
 			.AutoHeight()
@@ -156,7 +171,7 @@ END_SLATE_FUNCTION_BUILD_OPTIMIZATION
 
 
 BEGIN_SLATE_FUNCTION_BUILD_OPTIMIZATION
-TSharedRef<SWidget> SSequencerPlaylistPanel::Construct_Toolbar()
+TSharedRef<SWidget> SSequencerPlaylistPanel::Construct_LeftToolbar()
 {
 	FSequencerPlaylistsModule& Module = static_cast<FSequencerPlaylistsModule&>(ISequencerPlaylistsModule::Get());
 	FSlimHorizontalToolBarBuilder ToolBarBuilder(Module.GetCommandList(), FMultiBoxCustomization::None, nullptr, true);
@@ -191,6 +206,37 @@ TSharedRef<SWidget> SSequencerPlaylistPanel::Construct_Toolbar()
 END_SLATE_FUNCTION_BUILD_OPTIMIZATION
 
 
+BEGIN_SLATE_FUNCTION_BUILD_OPTIMIZATION
+TSharedRef<SWidget> SSequencerPlaylistPanel::Construct_RightToolbar()
+{
+	TSharedRef<SSequencerPlaylistPanel> This = SharedThis(this);
+
+	FSequencerPlaylistsModule& Module = static_cast<FSequencerPlaylistsModule&>(ISequencerPlaylistsModule::Get());
+	FSlimHorizontalToolBarBuilder ToolBarBuilder(Module.GetCommandList(), FMultiBoxCustomization::None, nullptr, true);
+
+	ToolBarBuilder.BeginSection("TriggerMode");
+	{
+		TSharedRef<SWidget> TriggerToggle = SNew(SCheckBox)
+			.Padding(FMargin(10.0f, 4.0f))
+			.ToolTipText(LOCTEXT("ToggleTriggerToolTip", "Enables and disables Trigger Mode"))
+			.Style(FEditorStyle::Get(), "ToggleButtonCheckbox")
+			.IsChecked_Lambda([This]() { return This->InTriggerMode() ? ECheckBoxState::Checked : ECheckBoxState::Unchecked; })
+			.OnCheckStateChanged_Lambda([This](ECheckBoxState InState) { This->bTriggerMode = (InState == ECheckBoxState::Checked); })
+			[
+				SNew(STextBlock)
+				.Font(FEditorStyle::Get().GetFontStyle("FontAwesome.16"))
+				.Text(FText::FromString(FString(TEXT("\xf0e7"))))
+			];
+
+		ToolBarBuilder.AddWidget(TriggerToggle);
+	}
+	ToolBarBuilder.EndSection();
+
+	return ToolBarBuilder.MakeWidget();
+}
+END_SLATE_FUNCTION_BUILD_OPTIMIZATION
+
+
 BEGIN_SLATE_FUNCTION_BUILD_OPTIMIZATION
 TSharedRef<SWidget> SSequencerPlaylistPanel::Construct_Transport()
 {
@@ -281,10 +327,11 @@ TSharedRef<SWidget> SSequencerPlaylistPanel::Construct_ItemListView()
 		.OnGenerateRow_Lambda([This = SharedThis(this)](TSharedPtr<FSequencerPlaylistRowData> InData, const TSharedRef<STableViewBase>& OwnerTableView)
 			{
 				return SNew(SSequencerPlaylistItemWidget, InData, OwnerTableView)
+					.TriggerMode(TAttribute<bool>::Create(TAttribute<bool>::FGetter::CreateSP(This, &SSequencerPlaylistPanel::InTriggerMode)))
 					.OnPlayClicked(This, &SSequencerPlaylistPanel::HandleClicked_Item_Play)
 					.OnStopClicked(This, &SSequencerPlaylistPanel::HandleClicked_Item_Stop)
 					.OnResetClicked(This, &SSequencerPlaylistPanel::HandleClicked_Item_Reset)
-					.OnRemoveClicked(This, &SSequencerPlaylistPanel::HandleClicked_Item_Remove )
+					.OnRemoveClicked(This, &SSequencerPlaylistPanel::HandleClicked_Item_Remove)
 					.OnIsPropertyVisible(This, &SSequencerPlaylistPanel::HandleItemDetailsIsPropertyVisible)
 					.OnCanAcceptDrop(This, &SSequencerPlaylistPanel::HandleCanAcceptDrop)
 					.OnAcceptDrop(This, &SSequencerPlaylistPanel::HandleAcceptDrop);
@@ -809,15 +856,23 @@ FSequencerPlaylistItemDragDropOp::~FSequencerPlaylistItemDragDropOp()
 }
 
 
+const FText SSequencerPlaylistItemWidget::PlayItemTooltipText(LOCTEXT("PlayItemButtonToolTip", "Play this item"));
+const FText SSequencerPlaylistItemWidget::StopItemTooltipText(LOCTEXT("StopItemButtonToolTip", "Stop playback of this item"));
+const FText SSequencerPlaylistItemWidget::ResetItemTooltipText(LOCTEXT("ResetItemButtonToolTip", "Reset this item"));
+
+
 void SSequencerPlaylistItemWidget::Construct(const FArguments& InArgs, TSharedPtr<FSequencerPlaylistRowData> InRowData, const TSharedRef<STableViewBase>& OwnerTableView)
 {
 	check(InRowData && InRowData->WeakItem.IsValid());
 	RowData = InRowData;
 
+	TriggerMode = InArgs._TriggerMode;
+
 	PlayClickedDelegate = InArgs._OnPlayClicked;
 	StopClickedDelegate = InArgs._OnStopClicked;
 	ResetClickedDelegate = InArgs._OnResetClicked;
 	RemoveClickedDelegate = InArgs._OnRemoveClicked;
+
 	IsPropertyVisibleDelegate = InArgs._OnIsPropertyVisible;
 
 	FSuperRowType::Construct(
@@ -829,6 +884,90 @@ void SSequencerPlaylistItemWidget::Construct(const FArguments& InArgs, TSharedPt
 }
 
 
+void SSequencerPlaylistItemWidget::ConstructChildren(ETableViewMode::Type InOwnerTableMode, const TAttribute<FMargin>& InPadding, const TSharedRef<SWidget>& InContent)
+{
+	// We wrap InContent with an overlay which facilitates dimming the row, obstructing hit test of
+	// the controls behind, and centering Trigger Mode transport controls over the row on hover.
+
+	TSharedRef<SSequencerPlaylistItemWidget> This = SharedThis(this);
+
+	Content = InContent;
+
+	ChildSlot
+	.Padding(InPadding)
+	[
+		SNew(SOverlay)
+		+ SOverlay::Slot()
+		[
+			InContent
+		]
+		+ SOverlay::Slot()
+		[
+			SNew(SImage)
+			.Visibility(this, &SSequencerPlaylistItemWidget::GetTriggerModeTransportVisibility)
+			.Image(FSequencerPlaylistsStyle::Get().GetBrush("SequencerPlaylists.Item.Dim"))
+		]
+		+ SOverlay::Slot()
+		.HAlign(HAlign_Center)
+		.VAlign(VAlign_Center)
+		[
+			SNew(SHorizontalBox)
+			.Visibility(this, &SSequencerPlaylistItemWidget::GetTriggerModeTransportVisibility)
+			+ SHorizontalBox::Slot()
+			.AutoWidth()
+			[
+				SNew(SButton)
+				.ButtonStyle(FSequencerPlaylistsStyle::Get(), "SequencerPlaylists.TransportButton.Play")
+				.ContentPadding(0)
+				.HAlign(HAlign_Center)
+				.VAlign(VAlign_Center)
+				.OnClicked_Lambda([This]() { return This->PlayClickedDelegate.Execute(This); })
+				.ToolTipText(PlayItemTooltipText)
+				[
+					SNew(SImage)
+					.Image(FSequencerPlaylistsStyle::Get().GetBrush("SequencerPlaylists.Play"))
+					.ColorAndOpacity(FSlateColor::UseForeground())
+				]
+			]
+			+ SHorizontalBox::Slot()
+			.AutoWidth()
+			[
+				SNew(SButton)
+				.ButtonStyle(FSequencerPlaylistsStyle::Get(), "SequencerPlaylists.TransportButton.Stop")
+				.ContentPadding(0)
+				.HAlign(HAlign_Center)
+				.VAlign(VAlign_Center)
+				.OnClicked_Lambda([This]() { return This->StopClickedDelegate.Execute(This); })
+				.ToolTipText(StopItemTooltipText)
+				[
+					SNew(SImage)
+					.Image(FSequencerPlaylistsStyle::Get().GetBrush("SequencerPlaylists.Stop"))
+					.ColorAndOpacity(FSlateColor::UseForeground())
+				]
+			]
+			+ SHorizontalBox::Slot()
+			.AutoWidth()
+			[
+				SNew(SButton)
+				.ButtonStyle(FSequencerPlaylistsStyle::Get(), "SequencerPlaylists.TransportButton.Reset")
+				.ContentPadding(FMargin(0.0f, 1.0f, 0.0f, 0.0f))
+				.HAlign(HAlign_Center)
+				.VAlign(VAlign_Center)
+				.OnClicked_Lambda([This]() { return This->ResetClickedDelegate.Execute(This); })
+				.ToolTipText(ResetItemTooltipText)
+				[
+					SNew(STextBlock)
+					.Font(FEditorStyle::Get().GetFontStyle("FontAwesome.14"))
+					.Text(FEditorFontGlyphs::Undo)
+				]
+			]
+		]
+	];
+
+	InnerContentSlot = &ChildSlot.AsSlot();
+}
+
+
 BEGIN_SLATE_FUNCTION_BUILD_OPTIMIZATION
 TSharedRef<SWidget> SSequencerPlaylistItemWidget::GenerateWidgetForColumn(const FName& ColumnName)
 {
@@ -839,11 +978,64 @@ TSharedRef<SWidget> SSequencerPlaylistItemWidget::GenerateWidgetForColumn(const
 		return SNullWidget::NullWidget;
 	}
 
+	TSharedRef<SSequencerPlaylistItemWidget> This = SharedThis(this);
+
 	static const FName PropertyEditorModuleName("PropertyEditor");
 	FPropertyEditorModule& PropertyEditorModule = FModuleManager::GetModuleChecked<FPropertyEditorModule>(PropertyEditorModuleName);
 	FSinglePropertyParams SinglePropParams;
 	SinglePropParams.NamePlacement = EPropertyNamePlacement::Hidden;
 
+	if (ColumnName == SSequencerPlaylistPanel::ColumnName_HoverTransport)
+	{
+		return SNew(SVerticalBox)
+			.Visibility(TAttribute<EVisibility>::Create(TAttribute<EVisibility>::FGetter::CreateSP(SharedThis(this), &SSequencerPlaylistItemWidget::GetHoverTransportCellVisibility)))
+			+ SVerticalBox::Slot()
+			[
+				SNew(SButton)
+				.ButtonStyle(FSequencerPlaylistsStyle::Get(), "SequencerPlaylists.HoverTransport.Play")
+				.ContentPadding(FMargin(1.0f, 3.0f, 0.0f, 0.0f))
+				.HAlign(HAlign_Center)
+				.VAlign(VAlign_Center)
+				.OnClicked_Lambda([This]() { return This->PlayClickedDelegate.Execute(This); })
+				.ToolTipText(PlayItemTooltipText)
+				[
+					SNew(SImage)
+					.Image(FSequencerPlaylistsStyle::Get().GetBrush("SequencerPlaylists.Play.Small"))
+					.ColorAndOpacity(FSlateColor::UseForeground())
+				]
+			]
+			+ SVerticalBox::Slot()
+			[
+				SNew(SButton)
+				.ButtonStyle(FSequencerPlaylistsStyle::Get(), "SequencerPlaylists.HoverTransport.Stop")
+				.ContentPadding(FMargin(0.0f, 0.0f, 0.0f, 1.0f))
+				.HAlign(HAlign_Center)
+				.VAlign(VAlign_Center)
+				.OnClicked_Lambda([This]() { return This->StopClickedDelegate.Execute(This); })
+				.ToolTipText(StopItemTooltipText)
+				[
+					SNew(SImage)
+					.Image(FSequencerPlaylistsStyle::Get().GetBrush("SequencerPlaylists.Stop.Small"))
+					.ColorAndOpacity(FSlateColor::UseForeground())
+				]
+			]
+			+ SVerticalBox::Slot()
+			[
+				SNew(SButton)
+				.ButtonStyle(FSequencerPlaylistsStyle::Get(), "SequencerPlaylists.HoverTransport.Reset")
+				.ContentPadding(FMargin(1.0f, 0.0f, 0.0f, 3.0f))
+				.HAlign(HAlign_Center)
+				.VAlign(VAlign_Center)
+				.OnClicked_Lambda([This]() { return This->ResetClickedDelegate.Execute(This); })
+				.ToolTipText(ResetItemTooltipText)
+				[
+					SNew(STextBlock)
+					.Font(FEditorStyle::Get().GetFontStyle("FontAwesome.9"))
+					.Text(FEditorFontGlyphs::Undo)
+				]
+			];
+	}
+
 	if (ColumnName == SSequencerPlaylistPanel::ColumnName_Items)
 	{
 		USequencerPlaylistItem_Sequence* SequenceItem =
@@ -901,7 +1093,7 @@ TSharedRef<SWidget> SSequencerPlaylistItemWidget::GenerateWidgetForColumn(const
 	{
 		return SAssignNew(DetailsAnchor, SMenuAnchor)
 			.Placement(MenuPlacement_MenuLeft)
-			.OnGetMenuContent(this, &SSequencerPlaylistItemWidget::EnsureSelectedAndBuildContextMenu)
+			.OnGetMenuContent(This, &SSequencerPlaylistItemWidget::EnsureSelectedAndBuildContextMenu)
 			[
 				SNew(SBox)
 				.HeightOverride(24.0f)
@@ -910,10 +1102,10 @@ TSharedRef<SWidget> SSequencerPlaylistItemWidget::GenerateWidgetForColumn(const
 					SNew(SButton)
 					.ButtonStyle(FEditorStyle::Get(), "HoverHintOnly")
 					.ContentPadding(0)
-					.OnClicked_Lambda([this]() { DetailsAnchor->SetIsOpen(!DetailsAnchor->IsOpen()); return FReply::Handled(); })
+					.OnClicked_Lambda([This]() { This->DetailsAnchor->SetIsOpen(!This->DetailsAnchor->IsOpen()); return FReply::Handled(); })
 					.VAlign(VAlign_Center)
 					.HAlign(HAlign_Center)
-					.Visibility_Lambda([this]() { return this->IsHovered() ? EVisibility::Visible : EVisibility::Hidden; })
+					.Visibility_Lambda([This]() { return This->IsHovered() ? EVisibility::Visible : EVisibility::Hidden; })
 					[
 						SNew(SImage)
 						.ColorAndOpacity(FSlateColor::UseForeground())
@@ -954,6 +1146,34 @@ FReply SSequencerPlaylistItemWidget::HandleDragDetected(const FGeometry& MyGeome
 }
 
 
+bool SSequencerPlaylistItemWidget::IsRowContentEnabled() const
+{
+	return GetTriggerModeTransportVisibility() == EVisibility::Hidden;
+}
+
+
+EVisibility SSequencerPlaylistItemWidget::GetTriggerModeTransportVisibility() const
+{
+	if (IsHovered() && InTriggerMode())
+	{
+		return EVisibility::Visible;
+	}
+
+	return EVisibility::Hidden;
+}
+
+
+EVisibility SSequencerPlaylistItemWidget::GetHoverTransportCellVisibility() const
+{
+	if (IsHovered() && !InTriggerMode())
+	{
+		return EVisibility::Visible;
+	}
+
+	return EVisibility::Hidden;
+}
+
+
 TSharedRef<SWidget> SSequencerPlaylistItemWidget::EnsureSelectedAndBuildContextMenu()
 {
 	TArray<UObject*> SelectedItems;
@@ -973,6 +1193,8 @@ TSharedRef<SWidget> SSequencerPlaylistItemWidget::EnsureSelectedAndBuildContextM
 			OwnerTable->Private_ClearSelection();
 			OwnerTable->Private_SetItemSelection(RowData, true, true);
 			OwnerTable->Private_SignalSelectionChanged(ESelectInfo::OnMouseClick);
+			SelectedItems.Empty(1);
+			SelectedItems.Add(GetItem());
 		}
 	}
 	else
diff --git a/Engine/Plugins/VirtualProduction/SequencerPlaylists/Source/SequencerPlaylists/Private/SequencerPlaylistsWidgets.h b/Engine/Plugins/VirtualProduction/SequencerPlaylists/Source/SequencerPlaylists/Private/SequencerPlaylistsWidgets.h
--- a/Engine/Plugins/VirtualProduction/SequencerPlaylists/Source/SequencerPlaylists/Private/SequencerPlaylistsWidgets.h
+++ b/Engine/Plugins/VirtualProduction/SequencerPlaylists/Source/SequencerPlaylists/Private/SequencerPlaylistsWidgets.h
@@ -50,8 +50,11 @@ public:
 
 	void Construct(const FArguments& InArgs, USequencerPlaylistPlayer* InPlayer);
 
+	bool InTriggerMode() const { return bTriggerMode; }
+
 private:
-	TSharedRef<SWidget> Construct_Toolbar();
+	TSharedRef<SWidget> Construct_LeftToolbar();
+	TSharedRef<SWidget> Construct_RightToolbar();
 	TSharedRef<SWidget> Construct_Transport();
 	TSharedRef<SWidget> Construct_AddSearchRow();
 	TSharedRef<SWidget> Construct_ItemListView();
@@ -86,6 +89,8 @@ private:
 private:
 	TWeakObjectPtr<USequencerPlaylistPlayer> WeakPlayer;
 
+	bool bTriggerMode = false;
+
 	TSharedPtr<SSearchBox> SearchBox;
 	TSharedPtr<TTextFilter<const FSequencerPlaylistRowData&>> SearchTextFilter;
 	TArray<TSharedPtr<FSequencerPlaylistRowData>> ItemRows;
@@ -120,7 +125,13 @@ DECLARE_DELEGATE_RetVal_OneParam(FReply, FOnClickedSequencerPlaylistItem, TShare
 
 class SSequencerPlaylistItemWidget : public SMultiColumnTableRow<TSharedPtr<FSequencerPlaylistRowData>>
 {
+	static const FText PlayItemTooltipText;
+	static const FText StopItemTooltipText;
+	static const FText ResetItemTooltipText;
+
 	SLATE_BEGIN_ARGS(SSequencerPlaylistItemWidget) {}
+		SLATE_ATTRIBUTE(bool, TriggerMode)
+
 		SLATE_EVENT(FOnClickedSequencerPlaylistItem, OnPlayClicked)
 		SLATE_EVENT(FOnClickedSequencerPlaylistItem, OnStopClicked)
 		SLATE_EVENT(FOnClickedSequencerPlaylistItem, OnResetClicked)
@@ -137,6 +148,10 @@ public:
 	const TSharedPtr<FSequencerPlaylistRowData>& GetRowData() { return RowData; }
 	USequencerPlaylistItem* GetItem() { return GetRowData() ? GetRowData()->WeakItem.Get() : nullptr; }
 
+	//~ Begin STableRow
+	void ConstructChildren(ETableViewMode::Type InOwnerTableMode, const TAttribute<FMargin>& InPadding, const TSharedRef<SWidget>& InContent) override;
+	//~ End STableRow
+
 	//~ Begin SMultiColumnTableRow
 	TSharedRef<SWidget> GenerateWidgetForColumn(const FName& InColumnName) override;
 	//~ End SMultiColumnTableRow
@@ -148,6 +163,11 @@ public:
 private:
 	FReply HandleDragDetected(const FGeometry& MyGeometry, const FPointerEvent& MouseEvent);
 
+	bool InTriggerMode() const { return TriggerMode.Get(); }
+	bool IsRowContentEnabled() const;
+	EVisibility GetTriggerModeTransportVisibility() const;
+	EVisibility GetHoverTransportCellVisibility() const;
+
 	TSharedRef<SWidget> EnsureSelectedAndBuildContextMenu();
 	TSharedRef<SWidget> BuildContextMenu(const TArray<UObject*>& SelectedItems);
 
@@ -155,6 +175,8 @@ private:
 	TSharedPtr<FSequencerPlaylistRowData> RowData;
 	TSharedPtr<SMenuAnchor> DetailsAnchor;
 
+	TAttribute<bool> TriggerMode;
+
 	FOnClickedSequencerPlaylistItem PlayClickedDelegate;
 	FOnClickedSequencerPlaylistItem StopClickedDelegate;
 	FOnClickedSequencerPlaylistItem ResetClickedDelegate;
