Add support for unbounded play when playing a sequencer playlist.

diff --git a/Engine/Plugins/VirtualProduction/SequencerPlaylists/Source/SequencerPlaylists/Private/SequencerPlaylistPlayer.cpp b/Engine/Plugins/VirtualProduction/SequencerPlaylists/Source/SequencerPlaylists/Private/SequencerPlaylistPlayer.cpp
--- a/Engine/Plugins/VirtualProduction/SequencerPlaylists/Source/SequencerPlaylists/Private/SequencerPlaylistPlayer.cpp
+++ b/Engine/Plugins/VirtualProduction/SequencerPlaylists/Source/SequencerPlaylists/Private/SequencerPlaylistPlayer.cpp
@@ -1,6 +1,9 @@
 // Copyright Epic Games, Inc. All Rights Reserved.
 
 #include "SequencerPlaylistPlayer.h"
+#include "ISequencer.h"
+#include "Misc/QualifiedFrameTime.h"
+#include "MovieSceneFwd.h"
 #include "SequencerPlaylist.h"
 #include "SequencerPlaylistItem.h"
 #include "SequencerPlaylistsLog.h"
@@ -11,11 +14,11 @@
 #include "LevelSequenceEditorBlueprintLibrary.h"
 #include "Recorder/TakeRecorder.h"
 #include "ScopedTransaction.h"
-
+#include "TakePreset.h"
+#include "TakeRecorderSettings.h"
 
 #define LOCTEXT_NAMESPACE "SequencerPlaylists"
 
-
 USequencerPlaylistPlayer::USequencerPlaylistPlayer(const FObjectInitializer& ObjectInitializer)
 	: Super(ObjectInitializer)
 {
@@ -51,6 +54,8 @@ bool USequencerPlaylistPlayer::PlayItem(USequencerPlaylistItem* Item)
 		return false;
 	}
 
+	EnterUnboundedPlayIfNotRecording();
+
 	FScopedTransaction Transaction(FText::Format(LOCTEXT("PlayItemTransaction", "Trigger playback of {0}"), Item->GetDisplayName()));
 	return GetCheckedItemPlayer(Item)->Play(Item);
 }
@@ -79,6 +84,119 @@ bool USequencerPlaylistPlayer::ResetItem(USequencerPlaylistItem* Item)
 	return GetCheckedItemPlayer(Item)->Reset(Item);
 }
 
+namespace UE::Private::PlaylistPlayer
+{
+
+TOptional<TRange<double>> ComputeNewRange(TSharedPtr<ISequencer>& Sequencer)
+{
+	FAnimatedRange Range = Sequencer->GetViewRange();
+	UMovieSceneSequence* Sequence  = Sequencer->GetRootMovieSceneSequence();
+	if (Sequence)
+	{
+		UMovieScene* MovieScene = Sequence->GetMovieScene();
+		if (MovieScene)
+		{
+			FFrameRate FrameRate = MovieScene->GetTickResolution();
+			FQualifiedFrameTime GlobalTime = Sequencer->GetGlobalTime();
+			FFrameTime CurrentFrameTime = GlobalTime.ConvertTo(FrameRate);
+			double CurrentTimeSeconds = FrameRate.AsSeconds(CurrentFrameTime) + 0.5f;
+			CurrentTimeSeconds = CurrentTimeSeconds > Range.GetUpperBoundValue() ? CurrentTimeSeconds : Range.GetUpperBoundValue();
+			TRange<double> NewRange(Range.GetLowerBoundValue(), CurrentTimeSeconds);
+			return NewRange;
+		}
+	}
+	return {};
+}
+
+void AdjustMovieSceneRangeForPlay(TSharedPtr<ISequencer>& Sequencer)
+{
+	check(Sequencer);
+
+	TOptional<TRange<double>> NewRange = ComputeNewRange(Sequencer);
+	if (NewRange)
+	{
+		Sequencer->SetViewRange(*NewRange, EViewRangeInterpolation::Immediate);
+		Sequencer->SetClampRange(Sequencer->GetViewRange());
+	}
+}
+
+FFrameTime GetFrameTime(UMovieScene* MovieScene, FQualifiedFrameTime GlobalTime)
+{
+	FFrameRate FrameRate = MovieScene->GetTickResolution();
+	return GlobalTime.ConvertTo(FrameRate);
+}
+
+void SetInfinitePlayRange(TSharedPtr<ISequencer>& Sequencer)
+{
+	UMovieSceneSequence* Sequence = Sequencer->GetRootMovieSceneSequence();
+	UMovieScene* MovieScene = Sequence->GetMovieScene();
+
+	TRange<FFrameNumber> Range = MovieScene->GetPlaybackRange();
+	// Set infinite playback range when starting recording. Playback range will be clamped to the bounds of the sections at the completion of the recording
+	MovieScene->SetPlaybackRange(TRange<FFrameNumber>(Range.GetLowerBoundValue(), TNumericLimits<int32>::Max() - 1), false);
+}
+
+void StopPlaybackAndAdjustTime(TSharedPtr<ISequencer>& Sequencer)
+{
+	check(Sequencer);
+
+	Sequencer->SetPlaybackStatus(EMovieScenePlayerStatus::Stopped);
+	UMovieSceneSequence* Sequence  = Sequencer->GetRootMovieSceneSequence();
+	if (Sequence)
+	{
+		UMovieScene* MovieScene = Sequence->GetMovieScene();
+		if (MovieScene)
+		{
+			FFrameTime CurrentFrameTime = GetFrameTime(MovieScene, Sequencer->GetGlobalTime());
+			TRange<FFrameNumber> Range = MovieScene->GetPlaybackRange();
+			// Set the playback range back to a closed interval.
+			//
+			MovieScene->SetPlaybackRange(TRange<FFrameNumber>(Range.GetLowerBoundValue(), CurrentFrameTime.GetFrame()), false);
+		}
+	}
+}
+
+}
+
+void USequencerPlaylistPlayer::Tick(float DeltaTime)
+{
+	TSharedPtr<ISequencer> Sequencer = GetSequencer();
+
+	if (Sequencer->GetPlaybackStatus() == EMovieScenePlayerStatus::Paused)
+	{
+		return;
+	}
+
+	if (Sequencer->GetPlaybackStatus() == EMovieScenePlayerStatus::Stopped)
+	{
+		PlaylistTicker = {};
+		UE::Private::PlaylistPlayer::StopPlaybackAndAdjustTime(Sequencer);
+		return;
+	}
+	// Handle a tick
+	UE::Private::PlaylistPlayer::AdjustMovieSceneRangeForPlay(Sequencer);
+}
+
+void USequencerPlaylistPlayer::EnterUnboundedPlayIfNotRecording()
+{
+	TSharedPtr<ISequencer> Sequencer = GetSequencer();
+	UTakeRecorder* Recorder = UTakeRecorder::GetActiveRecorder();
+
+	const bool bInRecorder = Recorder && Recorder->GetState() != ETakeRecorderState::Stopped;
+	if (!PlaylistTicker || !bInRecorder)
+	{
+		PlaylistTicker = MakeUnique<FTickablePlaylist>(this);
+	}
+
+	if (PlaylistTicker && Sequencer->GetPlaybackStatus() != EMovieScenePlayerStatus::Playing)
+	{
+		UE::Private::PlaylistPlayer::SetInfinitePlayRange(Sequencer);
+		Sequencer->SetPlaybackStatus(EMovieScenePlayerStatus::Playing);
+
+		// Tick once to set our playback range.
+		Tick(0.0);
+	}
+}
 
 bool USequencerPlaylistPlayer::PlayAll()
 {
@@ -87,6 +205,8 @@ bool USequencerPlaylistPlayer::PlayAll()
 		return false;
 	}
 
+	EnterUnboundedPlayIfNotRecording();
+
 	bool bResult = true;
 
 	FScopedTransaction Transaction(LOCTEXT("PlayAllTransaction", "Trigger playback of all items"));
@@ -106,6 +226,15 @@ bool USequencerPlaylistPlayer::StopAll()
 		return false;
 	}
 
+	UTakeRecorder* Recorder = UTakeRecorder::GetActiveRecorder();
+
+	TSharedPtr<ISequencer> Sequencer = GetSequencer();
+	const bool bInRecorder = Recorder && Recorder->GetState() != ETakeRecorderState::Stopped;
+	if (!bInRecorder && Sequencer->GetPlaybackStatus() == EMovieScenePlayerStatus::Playing)
+	{
+		UE::Private::PlaylistPlayer::StopPlaybackAndAdjustTime(Sequencer);
+	}
+
 	bool bResult = true;
 
 	FScopedTransaction Transaction(LOCTEXT("StopAllTransaction", "Stop playback of all items"));
@@ -114,6 +243,7 @@ bool USequencerPlaylistPlayer::StopAll()
 		bResult &= GetCheckedItemPlayer(Item)->Stop(Item);
 	}
 
+	PlaylistTicker = {};
 	return bResult;
 }
 
@@ -147,9 +277,14 @@ TSharedPtr<ISequencer> USequencerPlaylistPlayer::GetSequencer()
 	ULevelSequence* RootSequence = ULevelSequenceEditorBlueprintLibrary::GetCurrentLevelSequence();
 	if (!RootSequence)
 	{
-		// TODO: Instantiate sequencer with new empty take?
-		UE_LOG(LogSequencerPlaylists, Error, TEXT("USequencerPlaylistPlayer::GetSequencer: GetCurrentLevelSequence returned null"));
-		return nullptr;
+		UTakePreset* Preset = UTakePreset::AllocateTransientPreset(GetDefault<UTakeRecorderUserSettings>()->LastOpenedPreset.Get());
+
+		FScopedTransaction Transaction(LOCTEXT("CreateEmptyTake", "Create Empty Playlist Sequence"));
+
+		Preset->Modify();
+		Preset->CreateLevelSequence();
+
+		RootSequence = Preset->GetLevelSequence();
 	}
 
 	GEditor->GetEditorSubsystem<UAssetEditorSubsystem>()->OpenEditorForAsset(RootSequence);
diff --git a/Engine/Plugins/VirtualProduction/SequencerPlaylists/Source/SequencerPlaylists/Public/SequencerPlaylistPlayer.h b/Engine/Plugins/VirtualProduction/SequencerPlaylists/Source/SequencerPlaylists/Public/SequencerPlaylistPlayer.h
--- a/Engine/Plugins/VirtualProduction/SequencerPlaylists/Source/SequencerPlaylists/Public/SequencerPlaylistPlayer.h
+++ b/Engine/Plugins/VirtualProduction/SequencerPlaylists/Source/SequencerPlaylists/Public/SequencerPlaylistPlayer.h
@@ -5,6 +5,10 @@
 #include "CoreMinimal.h"
 #include "Templates/SubclassOf.h"
 #include "SequencerPlaylistItem.h"
+#include "Templates/UniquePtr.h"
+
+#include "Tickable.h"
+
 #include "SequencerPlaylistPlayer.generated.h"
 
 
@@ -12,6 +16,7 @@ class ISequencer;
 class ISequencerPlaylistItemPlayer;
 class USequencerPlaylist;
 class UTakeRecorder;
+class FTickablePlaylist;
 
 
 /** Controls playback of playlist items */
@@ -52,6 +57,16 @@ public:
 private:
 	TSharedPtr<ISequencer> GetSequencer();
 
+	/**
+	 * Enters unbounded play (like Take Recorder) if the user is not recording.
+	 */
+	void EnterUnboundedPlayIfNotRecording();
+
+	/**
+	 * Ticked by a tickable game object to performe any necessary time-sliced logic
+	 */
+	void Tick(float DeltaTime);
+
 	/** Centralizes additional common assumptions that are convenient for downstream code to make. */
 	TSharedPtr<ISequencer> GetValidatedSequencer();
 
@@ -61,6 +76,8 @@ private:
 
 	TSharedPtr<ISequencerPlaylistItemPlayer> GetCheckedItemPlayer(USequencerPlaylistItem* Item);
 
+	friend class FTickablePlaylist;
+
 private:
 	TMap<TSubclassOf<USequencerPlaylistItem>, TSharedRef<ISequencerPlaylistItemPlayer>> ItemPlayersByType;
 
@@ -71,4 +88,47 @@ private:
 	TWeakObjectPtr<UTakeRecorder> WeakRecorder;
 
 	TWeakPtr<ISequencer> WeakSequencer;
+
+	/** A tickable object that will do the playing in the event we are not using a take recorder. */
+	TUniquePtr<FTickablePlaylist> PlaylistTicker;
+};
+
+
+class FTickablePlaylist : public FTickableGameObject
+{
+public:
+	FTickablePlaylist(USequencerPlaylistPlayer *InPlayer)
+		: WeakPlaylistPlayer(InPlayer)
+	{}
+
+	virtual TStatId GetStatId() const override
+	{
+		RETURN_QUICK_DECLARE_CYCLE_STAT(FTickablePlaylist, STATGROUP_Tickables);
+	}
+
+	//Make sure it always ticks, otherwise we can miss recording, in particularly when time code is always increasing throughout the system.
+	virtual ETickableTickType GetTickableTickType() const override { return ETickableTickType::Always; }
+
+	virtual bool IsTickableInEditor() const override
+	{
+		return true;
+	}
+
+	virtual UWorld* GetTickableGameObjectWorld() const override
+	{
+		USequencerPlaylistPlayer* Playlist = WeakPlaylistPlayer.Get();
+		return Playlist ? Playlist->GetWorld() : nullptr;
+	}
+
+	virtual void Tick(float DeltaTime) override
+	{
+		if (USequencerPlaylistPlayer* Playlist = WeakPlaylistPlayer.Get())
+		{
+			Playlist->Tick(DeltaTime);
+		}
+	}
+
+private:
+	TWeakObjectPtr<USequencerPlaylistPlayer> WeakPlaylistPlayer;
+
 };
