Sequencer Playlists: Fix "Playback Speed" property on playlist items. Also fixes a missing transaction affecting first frame hold, and makes playlist descriptions AssetRegistrySearchable.

diff --git a/Engine/Plugins/VirtualProduction/SequencerPlaylists/Source/SequencerPlaylists/Private/SequencerPlaylistItem_Sequence.cpp b/Engine/Plugins/VirtualProduction/SequencerPlaylists/Source/SequencerPlaylists/Private/SequencerPlaylistItem_Sequence.cpp
--- a/Engine/Plugins/VirtualProduction/SequencerPlaylists/Source/SequencerPlaylists/Private/SequencerPlaylistItem_Sequence.cpp
+++ b/Engine/Plugins/VirtualProduction/SequencerPlaylists/Source/SequencerPlaylists/Private/SequencerPlaylistItem_Sequence.cpp
@@ -76,15 +76,19 @@ bool FSequencerPlaylistItemPlayer_Sequence::Play(USequencerPlaylistItem* Item)
 		ItemState.WeakHoldSection.Reset();
 	}
 
-	const FFrameNumber StartFrame = GlobalTime.Time.FloorToFrame();
-	const FFrameNumber SingleLoopDuration = ConvertFrameTime(PlayRange.Size<FFrameNumber>(),
-		PlayScene->GetTickResolution(), RootScene->GetTickResolution()).FloorToFrame();
+	const float TimeScale = FMath::Max(SMALL_NUMBER, SequenceItem->PlaybackSpeed);
+	const FFrameTime SingleLoopDuration = ConvertFrameTime(PlayRange.Size<FFrameTime>() / TimeScale,
+		PlayScene->GetTickResolution(), RootScene->GetTickResolution());
 
+	const FFrameNumber StartFrame = GlobalTime.Time.FloorToFrame();
 	const int32 MaxDuration = TNumericLimits<int32>::Max() - StartFrame.Value - 1;
-	const int32 Duration = FMath::Min(MaxDuration, (SingleLoopDuration * FMath::Max(1, SequenceItem->NumLoops + 1)).Value);
+	const int32 Duration = FMath::Min(MaxDuration,
+		(SingleLoopDuration * FMath::Max(1, SequenceItem->NumLoops + 1)).FloorToFrame().Value);
 
 	UMovieSceneSubSection* WorkingSubSection = WorkingTrack->AddSequence(SequenceItem->Sequence, StartFrame, Duration);
 
+	WorkingSubSection->Parameters.TimeScale = TimeScale;
+
 	if (SequenceItem->StartFrameOffset > 0)
 	{
 		WorkingSubSection->Parameters.StartFrameOffset = StartFrameOffset_RootTicks.FloorToFrame();
diff --git a/Engine/Plugins/VirtualProduction/SequencerPlaylists/Source/SequencerPlaylists/Private/SequencerPlaylistPlayer.cpp b/Engine/Plugins/VirtualProduction/SequencerPlaylists/Source/SequencerPlaylists/Private/SequencerPlaylistPlayer.cpp
--- a/Engine/Plugins/VirtualProduction/SequencerPlaylists/Source/SequencerPlaylists/Private/SequencerPlaylistPlayer.cpp
+++ b/Engine/Plugins/VirtualProduction/SequencerPlaylists/Source/SequencerPlaylists/Private/SequencerPlaylistPlayer.cpp
@@ -361,6 +361,8 @@ void USequencerPlaylistPlayer::OnTakeRecorderStarted(UTakeRecorder* InRecorder)
 
 	if (TSharedPtr<ISequencer> Sequencer = GetValidatedSequencer())
 	{
+		FScopedTransaction Transaction(LOCTEXT("TakeRecorderStartedTransaction", "Playlist - Take Recorder started"));
+
 		for (USequencerPlaylistItem* Item : Playlist->Items)
 		{
 			if (Item->bHoldAtFirstFrame)
@@ -382,6 +384,8 @@ void USequencerPlaylistPlayer::OnTakeRecorderStopped(UTakeRecorder* InRecorder)
 	// FIXME: Any sequences not already stopped end up a few frames too long; pass in explicit end frame?
 	if (TSharedPtr<ISequencer> Sequencer = GetValidatedSequencer())
 	{
+		FScopedTransaction Transaction(LOCTEXT("TakeRecorderStoppedTransaction", "Playlist - Take Recorder stopped"));
+
 		for (USequencerPlaylistItem* Item : Playlist->Items)
 		{
 			GetCheckedItemPlayer(Item)->Stop(Item);
diff --git a/Engine/Plugins/VirtualProduction/SequencerPlaylists/Source/SequencerPlaylists/Public/SequencerPlaylist.h b/Engine/Plugins/VirtualProduction/SequencerPlaylists/Source/SequencerPlaylists/Public/SequencerPlaylist.h
--- a/Engine/Plugins/VirtualProduction/SequencerPlaylists/Source/SequencerPlaylists/Public/SequencerPlaylist.h
+++ b/Engine/Plugins/VirtualProduction/SequencerPlaylists/Source/SequencerPlaylists/Public/SequencerPlaylist.h
@@ -15,7 +15,7 @@ class USequencerPlaylist : public UObject
 	GENERATED_BODY()
 
 public:
-	UPROPERTY(VisibleAnywhere, BlueprintReadWrite, Category="SequencerPlaylists")
+	UPROPERTY(VisibleAnywhere, AssetRegistrySearchable, BlueprintReadWrite, Category="SequencerPlaylists")
 	FText Description;
 
 	UPROPERTY(VisibleAnywhere, BlueprintReadWrite, Category="SequencerPlaylists")
