Sequencer Playlists: Immediately update working playlist name (and therefore UI widget) in response to "Save As."

diff --git a/Engine/Plugins/VirtualProduction/SequencerPlaylists/Source/SequencerPlaylists/Private/SequencerPlaylistsWidgets.cpp b/Engine/Plugins/VirtualProduction/SequencerPlaylists/Source/SequencerPlaylists/Private/SequencerPlaylistsWidgets.cpp
--- a/Engine/Plugins/VirtualProduction/SequencerPlaylists/Source/SequencerPlaylists/Private/SequencerPlaylistsWidgets.cpp
+++ b/Engine/Plugins/VirtualProduction/SequencerPlaylists/Source/SequencerPlaylists/Private/SequencerPlaylistsWidgets.cpp
@@ -547,6 +547,9 @@ void SSequencerPlaylistPanel::OnSavePlaylistAs()
 
 	if (NewPlaylist)
 	{
+		// Within the transient package; no collision.
+		Playlist->Rename(*NewPlaylist->GetName());
+
 		NewPlaylist->Description = Playlist->Description;
 
 		for (USequencerPlaylistItem* Item : Playlist->Items)
@@ -577,8 +580,7 @@ void SSequencerPlaylistPanel::OnLoadPlaylist(const FAssetData& InPreset)
 	{
 		FScopedTransaction Transaction(LOCTEXT("LoadPlaylistTransaction", "Load Sequencer Playlist"));
 
-		// FIXME: Within the transient package, which we'll never be loading from.
-		// This will go away when this is refactored into an asset editor. (UETOOL-4643)
+		// Within the transient package; no collision.
 		Playlist->Rename(*ImportedPlaylist->GetName());
 
 		Playlist->Description = ImportedPlaylist->Description;
