Sequencer Playlists: Warn and early out if play or hold/reset actions would create a circular reference between sequences.

diff --git a/Engine/Plugins/VirtualProduction/SequencerPlaylists/Source/SequencerPlaylists/Private/SequencerPlaylistItem_Sequence.cpp b/Engine/Plugins/VirtualProduction/SequencerPlaylists/Source/SequencerPlaylists/Private/SequencerPlaylistItem_Sequence.cpp
--- a/Engine/Plugins/VirtualProduction/SequencerPlaylists/Source/SequencerPlaylists/Private/SequencerPlaylistItem_Sequence.cpp
+++ b/Engine/Plugins/VirtualProduction/SequencerPlaylists/Source/SequencerPlaylists/Private/SequencerPlaylistItem_Sequence.cpp
@@ -5,11 +5,14 @@
 #include "SequencerPlaylistsLog.h"
 
 #include "ISequencer.h"
+#include "Framework/Notifications/NotificationManager.h"
 #include "LevelSequence.h"
 #include "MovieSceneFolder.h"
 #include "MovieSceneTimeHelpers.h"
 #include "Sections/MovieSceneSubSection.h"
+#include "TrackEditors/SubTrackEditorBase.h" // for FSubTrackEditorUtil::CanAddSubSequence
 #include "Tracks/MovieSceneSubTrack.h"
+#include "Widgets/Notifications/SNotificationList.h"
 
 
 #define LOCTEXT_NAMESPACE "SequencerPlaylists"
@@ -41,14 +44,22 @@ bool FSequencerPlaylistItemPlayer_Sequence::Play(USequencerPlaylistItem* Item)
 		return false;
 	}
 
-	FItemState& ItemState = GetOrCreateItemState(Item);
-
-	UMovieSceneSubTrack* WorkingTrack = GetOrCreateWorkingTrack(Item);
 	ULevelSequence* RootSequence = Cast<ULevelSequence>(Sequencer->GetRootMovieSceneSequence());
 	UMovieScene* RootScene = RootSequence->GetMovieScene();
 
+	if (!FSubTrackEditorUtil::CanAddSubSequence(RootSequence, *SequenceItem->Sequence))
+	{
+		FNotificationInfo Info(FText::Format(LOCTEXT("InvalidSequence", "Invalid level sequence {0}. There could be a circular dependency."), SequenceItem->Sequence->GetDisplayName()));
+		Info.bUseLargeFont = false;
+		FSlateNotificationManager::Get().AddNotification(Info);
+		return false;
+	}
+
 	RootSequence->Modify();
 
+	FItemState& ItemState = GetOrCreateItemState(Item);
+	UMovieSceneSubTrack* WorkingTrack = GetOrCreateWorkingTrack(Item);
+
 	const FQualifiedFrameTime GlobalTime = Sequencer->GetGlobalTime();
 
 	UMovieScene* PlayScene = SequenceItem->Sequence->GetMovieScene();
@@ -184,12 +195,21 @@ bool FSequencerPlaylistItemPlayer_Sequence::AddHold(USequencerPlaylistItem* Item
 		return false;
 	}
 
-	UMovieSceneSubTrack* WorkingTrack = GetOrCreateWorkingTrack(Item);
 	ULevelSequence* RootSequence = Cast<ULevelSequence>(Sequencer->GetRootMovieSceneSequence());
 	UMovieScene* RootScene = RootSequence->GetMovieScene();
 
+	if (!FSubTrackEditorUtil::CanAddSubSequence(RootSequence, *SequenceItem->Sequence))
+	{
+		FNotificationInfo Info(FText::Format(LOCTEXT("InvalidSequence", "Invalid level sequence {0}. There could be a circular dependency."), SequenceItem->Sequence->GetDisplayName()));
+		Info.bUseLargeFont = false;
+		FSlateNotificationManager::Get().AddNotification(Info);
+		return false;
+	}
+
 	RootSequence->Modify();
 
+	UMovieSceneSubTrack* WorkingTrack = GetOrCreateWorkingTrack(Item);
+
 	const FQualifiedFrameTime GlobalTime = Sequencer->GetGlobalTime();
 	const FFrameNumber StartFrame = GlobalTime.Time.FloorToFrame();
 	const int32 MaxDuration = TNumericLimits<int32>::Max() - StartFrame.Value - 1;
diff --git a/Engine/Plugins/VirtualProduction/SequencerPlaylists/Source/SequencerPlaylists/SequencerPlaylists.Build.cs b/Engine/Plugins/VirtualProduction/SequencerPlaylists/Source/SequencerPlaylists/SequencerPlaylists.Build.cs
--- a/Engine/Plugins/VirtualProduction/SequencerPlaylists/Source/SequencerPlaylists/SequencerPlaylists.Build.cs
+++ b/Engine/Plugins/VirtualProduction/SequencerPlaylists/Source/SequencerPlaylists/SequencerPlaylists.Build.cs
@@ -37,6 +37,7 @@ public class SequencerPlaylists : ModuleRules
 			"LevelSequenceEditor",
 			"TakesCore",
 			"ToolWidgets",
+			"MovieSceneTools",
 		});
 
 		DynamicallyLoadedModuleNames.AddRange(new string[] {
