Sequencer Playlists: Workaround for looped items not stopping correctly until after completing playback of the first loop: Check the section we're trying to stop, and if it's marked bCanLoop, but isn't long enough to have looped yet, unset bCanLoop.

Also:
 - Fixed some references to incorrect tick rates, which would cause problems if the root sequence and playlist item had different tick rates.
 - MovieSceneDataChange notifications hoisted out of play all/stop all loops, changing and documenting ISequencerPlaylistItemPlayer return value semantics in support of this goal.
 - Fixed a crash caused by "leaked" ticker if we exit the editor during unbounded playback.
 - Added some missing blueprint functionality.
 - Fixed some bad indentation in MovieSceneCompiledDataManager.cpp.

diff --git a/Engine/Plugins/VirtualProduction/SequencerPlaylists/Source/SequencerPlaylists/Private/SequencerPlaylistItem_Sequence.cpp b/Engine/Plugins/VirtualProduction/SequencerPlaylists/Source/SequencerPlaylists/Private/SequencerPlaylistItem_Sequence.cpp
--- a/Engine/Plugins/VirtualProduction/SequencerPlaylists/Source/SequencerPlaylists/Private/SequencerPlaylistItem_Sequence.cpp
+++ b/Engine/Plugins/VirtualProduction/SequencerPlaylists/Source/SequencerPlaylists/Private/SequencerPlaylistItem_Sequence.cpp
@@ -60,37 +60,36 @@ bool FSequencerPlaylistItemPlayer_Sequence::Play(USequencerPlaylistItem* Item)
 	FItemState& ItemState = GetOrCreateItemState(Item);
 	UMovieSceneSubTrack* WorkingTrack = GetOrCreateWorkingTrack(Item);
 
-	const FQualifiedFrameTime GlobalTime = Sequencer->GetGlobalTime();
+	if (UMovieSceneSubSection* HoldSection = ItemState.WeakHoldSection.Get())
+	{
+		EndSection(HoldSection);
+		ItemState.WeakHoldSection.Reset();
+	}
 
 	UMovieScene* PlayScene = SequenceItem->Sequence->GetMovieScene();
 	TRange<FFrameNumber> PlayRange = PlayScene->GetPlaybackRange();
 
-	const FFrameTime StartFrameOffset_RootTicks = ConvertFrameTime(FFrameNumber(SequenceItem->StartFrameOffset),
-		PlayScene->GetDisplayRate(), RootScene->GetTickResolution());
+	const FFrameTime StartFrameOffset_Ticks = ConvertFrameTime(FFrameNumber(SequenceItem->StartFrameOffset),
+		PlayScene->GetDisplayRate(), PlayScene->GetTickResolution());
 
-	const FFrameTime EndFrameOffset_RootTicks = ConvertFrameTime(FFrameNumber(SequenceItem->EndFrameOffset),
-		PlayScene->GetDisplayRate(), RootScene->GetTickResolution());
+	const FFrameTime EndFrameOffset_Ticks = ConvertFrameTime(FFrameNumber(SequenceItem->EndFrameOffset),
+		PlayScene->GetDisplayRate(), PlayScene->GetTickResolution());
 
 	if (SequenceItem->StartFrameOffset > 0)
 	{
-		PlayRange.SetLowerBoundValue(PlayRange.GetLowerBoundValue() + StartFrameOffset_RootTicks.FloorToFrame());
+		PlayRange.SetLowerBoundValue(PlayRange.GetLowerBoundValue() + StartFrameOffset_Ticks.FloorToFrame());
 	}
 
 	if (SequenceItem->EndFrameOffset > 0)
 	{
-		PlayRange.SetUpperBoundValue(PlayRange.GetUpperBoundValue() - EndFrameOffset_RootTicks.FloorToFrame());
-	}
-
-	if (UMovieSceneSubSection* HoldSection = ItemState.WeakHoldSection.Get())
-	{
-		EndSection(HoldSection);
-		ItemState.WeakHoldSection.Reset();
+		PlayRange.SetUpperBoundValue(PlayRange.GetUpperBoundValue() - EndFrameOffset_Ticks.FloorToFrame());
 	}
 
 	const float TimeScale = FMath::Max(SMALL_NUMBER, SequenceItem->PlaybackSpeed);
 	const FFrameTime SingleLoopDuration = ConvertFrameTime(PlayRange.Size<FFrameTime>() / TimeScale,
 		PlayScene->GetTickResolution(), RootScene->GetTickResolution());
 
+	const FQualifiedFrameTime GlobalTime = Sequencer->GetGlobalTime();
 	const FFrameNumber StartFrame = GlobalTime.Time.FloorToFrame();
 	const int32 MaxDuration = TNumericLimits<int32>::Max() - StartFrame.Value - 1;
 	const int32 Duration = FMath::Min(MaxDuration,
@@ -102,12 +101,12 @@ bool FSequencerPlaylistItemPlayer_Sequence::Play(USequencerPlaylistItem* Item)
 
 	if (SequenceItem->StartFrameOffset > 0)
 	{
-		WorkingSubSection->Parameters.StartFrameOffset = StartFrameOffset_RootTicks.FloorToFrame();
+		WorkingSubSection->Parameters.StartFrameOffset = StartFrameOffset_Ticks.FloorToFrame();
 	}
 
 	if (SequenceItem->EndFrameOffset > 0)
 	{
-		WorkingSubSection->Parameters.EndFrameOffset = EndFrameOffset_RootTicks.FloorToFrame();
+		WorkingSubSection->Parameters.EndFrameOffset = EndFrameOffset_Ticks.FloorToFrame();
 	}
 
 	if (SequenceItem->NumLoops != 0)
@@ -117,8 +116,6 @@ bool FSequencerPlaylistItemPlayer_Sequence::Play(USequencerPlaylistItem* Item)
 
 	ItemState.WeakPlaySections.Add(WorkingSubSection);
 
-	Sequencer->NotifyMovieSceneDataChanged(EMovieSceneDataChangeType::Unknown);
-
 	return true;
 }
 
@@ -136,7 +133,7 @@ bool FSequencerPlaylistItemPlayer_Sequence::Stop(USequencerPlaylistItem* Item)
 
 	if (ItemState.WeakPlaySections.Num() == 0 && HoldSection == nullptr)
 	{
-		return true;
+		return false;
 	}
 
 	if (HoldSection)
@@ -159,20 +156,6 @@ bool FSequencerPlaylistItemPlayer_Sequence::Stop(USequencerPlaylistItem* Item)
 }
 
 
-bool FSequencerPlaylistItemPlayer_Sequence::Reset(USequencerPlaylistItem* Item)
-{
-	if (Item->bHoldAtFirstFrame)
-	{
-		Stop(Item);
-		return AddHold(Item);
-	}
-	else
-	{
-		return Stop(Item);
-	}
-}
-
-
 bool FSequencerPlaylistItemPlayer_Sequence::AddHold(USequencerPlaylistItem* Item)
 {
 	USequencerPlaylistItem_Sequence* SequenceItem = CastChecked<USequencerPlaylistItem_Sequence>(Item);
@@ -189,7 +172,6 @@ bool FSequencerPlaylistItemPlayer_Sequence::AddHold(USequencerPlaylistItem* Item
 
 	FItemState& ItemState = GetOrCreateItemState(Item);
 	UMovieSceneSubSection* HoldSection = ItemState.WeakHoldSection.Get();
-
 	if (HoldSection)
 	{
 		return false;
@@ -210,30 +192,45 @@ bool FSequencerPlaylistItemPlayer_Sequence::AddHold(USequencerPlaylistItem* Item
 
 	UMovieSceneSubTrack* WorkingTrack = GetOrCreateWorkingTrack(Item);
 
+	ULevelSequence* HoldSequence = SequenceItem->Sequence;
+	UMovieScene* HoldScene = HoldSequence->GetMovieScene();
+
 	const FQualifiedFrameTime GlobalTime = Sequencer->GetGlobalTime();
 	const FFrameNumber StartFrame = GlobalTime.Time.FloorToFrame();
 	const int32 MaxDuration = TNumericLimits<int32>::Max() - StartFrame.Value - 1;
-	HoldSection = WorkingTrack->AddSequence(SequenceItem->Sequence, StartFrame, MaxDuration);
+	HoldSection = WorkingTrack->AddSequence(HoldSequence, StartFrame, MaxDuration);
 
 	if (SequenceItem->StartFrameOffset > 0)
 	{
-		const FFrameTime StartFrameOffset_RootTicks = ConvertFrameTime(
+		const FFrameTime StartFrameOffset_Ticks = ConvertFrameTime(
 			FFrameNumber(SequenceItem->StartFrameOffset),
-			SequenceItem->Sequence->GetMovieScene()->GetDisplayRate(),
-			RootScene->GetTickResolution());
-		HoldSection->Parameters.StartFrameOffset = StartFrameOffset_RootTicks.FloorToFrame();
+			HoldScene->GetDisplayRate(),
+			HoldScene->GetTickResolution());
+		HoldSection->Parameters.StartFrameOffset = StartFrameOffset_Ticks.FloorToFrame();
 	}
 
 	HoldSection->Parameters.TimeScale = 0.f;
 
 	ItemState.WeakHoldSection = HoldSection;
 
-	Sequencer->NotifyMovieSceneDataChanged(EMovieSceneDataChangeType::Unknown);
-
 	return true;
 }
 
 
+bool FSequencerPlaylistItemPlayer_Sequence::Reset(USequencerPlaylistItem* Item)
+{
+	if (Item->bHoldAtFirstFrame)
+	{
+		Stop(Item);
+		return AddHold(Item);
+	}
+	else
+	{
+		return Stop(Item);
+	}
+}
+
+
 UMovieSceneSubTrack* FSequencerPlaylistItemPlayer_Sequence::GetOrCreateWorkingTrack(USequencerPlaylistItem* Item)
 {
 	USequencerPlaylistItem_Sequence* SequenceItem = CastChecked<USequencerPlaylistItem_Sequence>(Item);
@@ -323,10 +320,28 @@ void FSequencerPlaylistItemPlayer_Sequence::EndSection(UMovieSceneSection* Secti
 	UMovieScene* SectionScene = Section->GetTypedOuter<UMovieScene>();
 	check(SectionScene);
 
-	const FFrameTime SectionNow = Sequencer->GetGlobalTime().ConvertTo(SectionScene->GetTickResolution());
+	const FQualifiedFrameTime GlobalTime = Sequencer->GetGlobalTime();
+	const FFrameTime SectionNow = GlobalTime.ConvertTo(SectionScene->GetTickResolution());
 	if (Section->IsTimeWithinSection(SectionNow.FloorToFrame()))
 	{
 		Section->SetEndFrame(TRangeBound<FFrameNumber>::Exclusive(SectionNow.FloorToFrame()));
+
+		// Workaround for not being able to interrupt first playthrough if bCanLoop == true.
+		if (UMovieSceneSubSection* SubSection = Cast<UMovieSceneSubSection>(Section))
+		{
+			if (SubSection->Parameters.bCanLoop)
+			{
+				// Calc whether we've looped; if not, set bCanLoop back to false.
+				FFrameTime InnerTime;
+				FMovieSceneWarpCounter WarpCounter;
+				SubSection->OuterToInnerTransform().TransformTime(GlobalTime.Time, InnerTime, WarpCounter);
+				const uint32 LastWarpCount = WarpCounter.LastWarpCount();
+				if (LastWarpCount == 0)
+				{
+					SubSection->Parameters.bCanLoop = false;
+				}
+			}
+		}
 	}
 }
 
diff --git a/Engine/Plugins/VirtualProduction/SequencerPlaylists/Source/SequencerPlaylists/Private/SequencerPlaylistPlayer.cpp b/Engine/Plugins/VirtualProduction/SequencerPlaylists/Source/SequencerPlaylists/Private/SequencerPlaylistPlayer.cpp
--- a/Engine/Plugins/VirtualProduction/SequencerPlaylists/Source/SequencerPlaylists/Private/SequencerPlaylistPlayer.cpp
+++ b/Engine/Plugins/VirtualProduction/SequencerPlaylists/Source/SequencerPlaylists/Private/SequencerPlaylistPlayer.cpp
@@ -19,6 +19,7 @@
 
 #define LOCTEXT_NAMESPACE "SequencerPlaylists"
 
+
 USequencerPlaylistPlayer::USequencerPlaylistPlayer(const FObjectInitializer& ObjectInitializer)
 	: Super(ObjectInitializer)
 {
@@ -37,6 +38,8 @@ void USequencerPlaylistPlayer::BeginDestroy()
 {
 	Super::BeginDestroy();
 
+	PlaylistTicker = {};
+
 	UTakeRecorder::OnRecordingInitialized().RemoveAll(this);
 
 	if (UTakeRecorder* BoundRecorder = WeakRecorder.Get())
@@ -47,6 +50,13 @@ void USequencerPlaylistPlayer::BeginDestroy()
 }
 
 
+//static
+USequencerPlaylistPlayer* USequencerPlaylistPlayer::GetDefaultPlayer()
+{
+	return ISequencerPlaylistsModule::Get().GetDefaultPlayer();
+}
+
+
 bool USequencerPlaylistPlayer::PlayItem(USequencerPlaylistItem* Item)
 {
 	if (!Item)
@@ -150,7 +160,6 @@ void StopPlaybackAndAdjustTime(TSharedPtr<ISequencer>& Sequencer)
 			FFrameTime CurrentFrameTime = GetFrameTime(MovieScene, Sequencer->GetGlobalTime());
 			TRange<FFrameNumber> Range = MovieScene->GetPlaybackRange();
 			// Set the playback range back to a closed interval.
-			//
 			MovieScene->SetPlaybackRange(TRange<FFrameNumber>(Range.GetLowerBoundValue(), CurrentFrameTime.GetFrame()), false);
 		}
 	}
@@ -207,15 +216,20 @@ bool USequencerPlaylistPlayer::PlayAll()
 
 	EnterUnboundedPlayIfNotRecording();
 
-	bool bResult = true;
+	bool bAnyChange = false;
 
 	FScopedTransaction Transaction(LOCTEXT("PlayAllTransaction", "Trigger playback of all items"));
 	for (USequencerPlaylistItem* Item : Playlist->Items)
 	{
-		bResult &= GetCheckedItemPlayer(Item)->Play(Item);
+		bAnyChange |= GetCheckedItemPlayer(Item)->Play(Item);
 	}
 
-	return bResult;
+	if (bAnyChange)
+	{
+		GetSequencer()->NotifyMovieSceneDataChanged(EMovieSceneDataChangeType::RefreshAllImmediately);
+	}
+
+	return bAnyChange;
 }
 
 
@@ -235,16 +249,22 @@ bool USequencerPlaylistPlayer::StopAll()
 		UE::Private::PlaylistPlayer::StopPlaybackAndAdjustTime(Sequencer);
 	}
 
-	bool bResult = true;
+	bool bAnyChange = false;
 
 	FScopedTransaction Transaction(LOCTEXT("StopAllTransaction", "Stop playback of all items"));
 	for (USequencerPlaylistItem* Item : Playlist->Items)
 	{
-		bResult &= GetCheckedItemPlayer(Item)->Stop(Item);
+		bAnyChange |= GetCheckedItemPlayer(Item)->Stop(Item);
+	}
+
+	if (bAnyChange)
+	{
+		Sequencer->NotifyMovieSceneDataChanged(EMovieSceneDataChangeType::RefreshAllImmediately);
 	}
 
 	PlaylistTicker = {};
-	return bResult;
+
+	return bAnyChange;
 }
 
 
@@ -255,15 +275,20 @@ bool USequencerPlaylistPlayer::ResetAll()
 		return false;
 	}
 
-	bool bResult = true;
+	bool bAnyChange = false;
 
 	FScopedTransaction Transaction(LOCTEXT("ResetAllTransaction", "Reset playback of all items"));
 	for (USequencerPlaylistItem* Item : Playlist->Items)
 	{
-		bResult &= GetCheckedItemPlayer(Item)->Reset(Item);
+		bAnyChange |= GetCheckedItemPlayer(Item)->Reset(Item);
+	}
+
+	if (bAnyChange)
+	{
+		GetSequencer()->NotifyMovieSceneDataChanged(EMovieSceneDataChangeType::RefreshAllImmediately);
 	}
 
-	return bResult;
+	return bAnyChange;
 }
 
 
diff --git a/Engine/Plugins/VirtualProduction/SequencerPlaylists/Source/SequencerPlaylists/Public/ISequencerPlaylistsModule.h b/Engine/Plugins/VirtualProduction/SequencerPlaylists/Source/SequencerPlaylists/Public/ISequencerPlaylistsModule.h
--- a/Engine/Plugins/VirtualProduction/SequencerPlaylists/Source/SequencerPlaylists/Public/ISequencerPlaylistsModule.h
+++ b/Engine/Plugins/VirtualProduction/SequencerPlaylists/Source/SequencerPlaylists/Public/ISequencerPlaylistsModule.h
@@ -13,15 +13,36 @@ class USequencerPlaylistItem;
 class USequencerPlaylistPlayer;
 
 
+/** Interface for derived classes that play one or more subclasses of USequencerPlaylistItem. */
 class ISequencerPlaylistItemPlayer
 {
 public:
 	virtual ~ISequencerPlaylistItemPlayer() {}
 
+	/**
+	 * Initiate playback of the specified item.
+	 * @return True if the current sequence was modified, otherwise false.
+	 */
 	virtual bool Play(USequencerPlaylistItem* Item) = 0;
+
+	/**
+	 * Halt any current playback of the specified item.
+	 * @return True if the current sequence was modified, otherwise false.
+	 */
 	virtual bool Stop(USequencerPlaylistItem* Item) = 0;
-	virtual bool Reset(USequencerPlaylistItem* Item) = 0;
+
+	/**
+	 * Adds a first frame hold section for the specified item.
+	 * @return True if the current sequence was modified, otherwise false.
+	 */
 	virtual bool AddHold(USequencerPlaylistItem* Item) = 0;
+
+	/**
+	 * Reset the specified item.
+	 * This is essentially a Stop(), followed by an AddHold() if applicable.
+	 * @return True if the current sequence was modified, otherwise false.
+	 */
+	virtual bool Reset(USequencerPlaylistItem* Item) = 0;
 };
 
 
diff --git a/Engine/Plugins/VirtualProduction/SequencerPlaylists/Source/SequencerPlaylists/Public/SequencerPlaylistPlayer.h b/Engine/Plugins/VirtualProduction/SequencerPlaylists/Source/SequencerPlaylists/Public/SequencerPlaylistPlayer.h
--- a/Engine/Plugins/VirtualProduction/SequencerPlaylists/Source/SequencerPlaylists/Public/SequencerPlaylistPlayer.h
+++ b/Engine/Plugins/VirtualProduction/SequencerPlaylists/Source/SequencerPlaylists/Public/SequencerPlaylistPlayer.h
@@ -32,9 +32,15 @@ public:
 	void BeginDestroy() override;
 	//~ End UObject interface
 
+	UFUNCTION(BlueprintCallable, Category="SequencerPlaylists")
 	void SetPlaylist(USequencerPlaylist* InPlaylist) { Playlist = InPlaylist; }
+
+	UFUNCTION(BlueprintPure, Category="SequencerPlaylists")
 	USequencerPlaylist* GetPlaylist() { return Playlist; }
 
+	UFUNCTION(BlueprintPure, meta=(DisplayName="Get Default Sequencer Playlist Player"), Category="SequencerPlaylists")
+	static USequencerPlaylistPlayer* GetDefaultPlayer();
+
 public:
 	UFUNCTION(BlueprintCallable, Category="SequencerPlaylists")
 	bool PlayItem(USequencerPlaylistItem* Item);
diff --git a/Engine/Source/Runtime/MovieScene/Private/Compilation/MovieSceneCompiledDataManager.cpp b/Engine/Source/Runtime/MovieScene/Private/Compilation/MovieSceneCompiledDataManager.cpp
--- a/Engine/Source/Runtime/MovieScene/Private/Compilation/MovieSceneCompiledDataManager.cpp
+++ b/Engine/Source/Runtime/MovieScene/Private/Compilation/MovieSceneCompiledDataManager.cpp
@@ -1576,10 +1576,10 @@ void UMovieSceneCompiledDataManager::PopulateSubSequenceTree(UMovieSceneSubTrack
 			// Recurse into the sub sequence
 			RootPath->PushGeneration(SubSequenceID, SubData->DeterministicSequenceID);
 			{
-			PopulateSubSequenceTree(SubData->GetSequence(), SubParams, RootPath, InOutHierarchy);
+				PopulateSubSequenceTree(SubData->GetSequence(), SubParams, RootPath, InOutHierarchy);
+			}
+			RootPath->PopGenerations(1);
 		}
-		RootPath->PopGenerations(1);
-	}
 		else
 		{
 			// The section is looping so we need to add its contents to the tree as many times as it has loops.
