Sequencer Playlists: Mark USequencerPlaylistItem properties BlueprintReadWrite.

diff --git a/Engine/Plugins/VirtualProduction/SequencerPlaylists/Source/SequencerPlaylists/Public/SequencerPlaylistItem.h b/Engine/Plugins/VirtualProduction/SequencerPlaylists/Source/SequencerPlaylists/Public/SequencerPlaylistItem.h
--- a/Engine/Plugins/VirtualProduction/SequencerPlaylists/Source/SequencerPlaylists/Public/SequencerPlaylistItem.h
+++ b/Engine/Plugins/VirtualProduction/SequencerPlaylists/Source/SequencerPlaylists/Public/SequencerPlaylistItem.h
@@ -19,23 +19,23 @@ public:
 	virtual FText GetDisplayName() PURE_VIRTUAL(USequencerPlaylistItem::GetDisplayName, return FText::GetEmpty(); )
 
 public:
-	UPROPERTY(EditAnywhere, Category="SequencerPlaylists", meta=(NoResetToDefault))
+	UPROPERTY(EditAnywhere, BlueprintReadWrite, Category="SequencerPlaylists", meta=(NoResetToDefault))
 	int32 StartFrameOffset;
 
-	UPROPERTY(EditAnywhere, Category="SequencerPlaylists", meta=(NoResetToDefault))
+	UPROPERTY(EditAnywhere, BlueprintReadWrite, Category="SequencerPlaylists", meta=(NoResetToDefault))
 	int32 EndFrameOffset;
 
 	/**
 	 * If true, the sequence will be inserted immediately on recording start and any time Reset()
 	 * is called, paused at the first frame indefinitely until either triggered or stopped.
 	 */
-	UPROPERTY(EditAnywhere, Category="SequencerPlaylists", meta=(NoResetToDefault))
+	UPROPERTY(EditAnywhere, BlueprintReadWrite, Category="SequencerPlaylists", meta=(NoResetToDefault))
 	bool bHoldAtFirstFrame;
 
 	/** 0 is single playthrough, >= 1 is (n+1) playthroughs. */
-	UPROPERTY(EditAnywhere, Category="SequencerPlaylists", meta=(NoResetToDefault))
+	UPROPERTY(EditAnywhere, BlueprintReadWrite, Category="SequencerPlaylists", meta=(NoResetToDefault))
 	int32 NumLoops;
 
-	UPROPERTY(EditAnywhere, Category="SequencerPlaylists", meta=(NoResetToDefault))
+	UPROPERTY(EditAnywhere, BlueprintReadWrite, Category="SequencerPlaylists", meta=(NoResetToDefault))
 	float PlaybackSpeed = 1.0f;
 };
