Sequencer Playlists: Fix for stop button not being persistently visible on playing items.

Also addresses some closely related issues:
 - Fix for stale item state - clear item player state in response to sequencer stop event.
 - Call Pause() instead of forcing Stopped status when ending boundless playback, and stop all items first.
 - If stopping an item would result in a zero length section, remove the section entirely instead.
 - Don't attempt to open a new sequencer in USequencerPlaylistPlayer::Tick if our weak pointer to the expected sequencer is invalid.

diff --git a/Engine/Plugins/VirtualProduction/SequencerPlaylists/Source/SequencerPlaylists/Private/SequencerPlaylistItem_Sequence.cpp b/Engine/Plugins/VirtualProduction/SequencerPlaylists/Source/SequencerPlaylists/Private/SequencerPlaylistItem_Sequence.cpp
--- a/Engine/Plugins/VirtualProduction/SequencerPlaylists/Source/SequencerPlaylists/Private/SequencerPlaylistItem_Sequence.cpp
+++ b/Engine/Plugins/VirtualProduction/SequencerPlaylists/Source/SequencerPlaylists/Private/SequencerPlaylistItem_Sequence.cpp
@@ -27,6 +27,16 @@ FText USequencerPlaylistItem_Sequence::GetDisplayName()
 FSequencerPlaylistItemPlayer_Sequence::FSequencerPlaylistItemPlayer_Sequence(TSharedRef<ISequencer> Sequencer)
 	: WeakSequencer(Sequencer)
 {
+	Sequencer->OnStopEvent().AddRaw(this, &FSequencerPlaylistItemPlayer_Sequence::ClearItemStates);
+}
+
+
+FSequencerPlaylistItemPlayer_Sequence::~FSequencerPlaylistItemPlayer_Sequence()
+{
+	if (TSharedPtr<ISequencer> Sequencer = WeakSequencer.Pin())
+	{
+		Sequencer->OnStopEvent().RemoveAll(this);
+	}
 }
 
 
@@ -57,7 +67,7 @@ bool FSequencerPlaylistItemPlayer_Sequence::Play(USequencerPlaylistItem* Item)
 
 	RootSequence->Modify();
 
-	FItemState& ItemState = GetOrCreateItemState(Item);
+	FItemState& ItemState = ItemStates.FindOrAdd(Item);
 	UMovieSceneSubTrack* WorkingTrack = GetOrCreateWorkingTrack(Item);
 
 	if (UMovieSceneSubSection* HoldSection = ItemState.WeakHoldSection.Get())
@@ -95,6 +105,8 @@ bool FSequencerPlaylistItemPlayer_Sequence::Play(USequencerPlaylistItem* Item)
 	const int32 Duration = FMath::Min(MaxDuration,
 		(SingleLoopDuration * FMath::Max(1, SequenceItem->NumLoops + 1)).FloorToFrame().Value);
 
+	ItemState.PlayingUntil_RootTicks = FMath::Max(StartFrame.Value + Duration, ItemState.PlayingUntil_RootTicks);
+
 	UMovieSceneSubSection* WorkingSubSection = WorkingTrack->AddSequence(SequenceItem->Sequence, StartFrame, Duration);
 
 	WorkingSubSection->Parameters.TimeScale = TimeScale;
@@ -128,7 +140,7 @@ bool FSequencerPlaylistItemPlayer_Sequence::Stop(USequencerPlaylistItem* Item)
 		return false;
 	}
 
-	FItemState& ItemState = GetOrCreateItemState(Item);
+	FItemState& ItemState = ItemStates.FindOrAdd(Item);
 	UMovieSceneSubSection* HoldSection = ItemState.WeakHoldSection.Get();
 
 	if (ItemState.WeakPlaySections.Num() == 0 && HoldSection == nullptr)
@@ -152,6 +164,12 @@ bool FSequencerPlaylistItemPlayer_Sequence::Stop(USequencerPlaylistItem* Item)
 
 	ItemState.WeakPlaySections.Empty();
 
+	if (TSharedPtr<ISequencer> Sequencer = WeakSequencer.Pin())
+	{
+		const FFrameNumber Now = Sequencer->GetGlobalTime().Time.FloorToFrame();
+		ItemState.PlayingUntil_RootTicks = FMath::Min(Now.Value, ItemState.PlayingUntil_RootTicks);
+	}
+
 	return true;
 }
 
@@ -170,7 +188,7 @@ bool FSequencerPlaylistItemPlayer_Sequence::AddHold(USequencerPlaylistItem* Item
 		return false;
 	}
 
-	FItemState& ItemState = GetOrCreateItemState(Item);
+	FItemState& ItemState = ItemStates.FindOrAdd(Item);
 	UMovieSceneSubSection* HoldSection = ItemState.WeakHoldSection.Get();
 	if (HoldSection)
 	{
@@ -199,6 +217,7 @@ bool FSequencerPlaylistItemPlayer_Sequence::AddHold(USequencerPlaylistItem* Item
 	const FFrameNumber StartFrame = GlobalTime.Time.FloorToFrame();
 	const int32 MaxDuration = TNumericLimits<int32>::Max() - StartFrame.Value - 1;
 	HoldSection = WorkingTrack->AddSequence(HoldSequence, StartFrame, MaxDuration);
+	ItemState.PlayingUntil_RootTicks = TNumericLimits<int32>::Max();
 
 	if (SequenceItem->StartFrameOffset > 0)
 	{
@@ -231,6 +250,36 @@ bool FSequencerPlaylistItemPlayer_Sequence::Reset(USequencerPlaylistItem* Item)
 }
 
 
+bool FSequencerPlaylistItemPlayer_Sequence::IsPlaying(USequencerPlaylistItem* Item) const
+{
+	USequencerPlaylistItem_Sequence* SequenceItem = CastChecked<USequencerPlaylistItem_Sequence>(Item);
+	if (!SequenceItem || !SequenceItem->Sequence)
+	{
+		return false;
+	}
+
+	TSharedPtr<ISequencer> Sequencer = WeakSequencer.Pin();
+	if (!ensure(Sequencer))
+	{
+		return false;
+	}
+
+	const FItemState* ItemState = ItemStates.Find(Item);
+	if (!ItemState)
+	{
+		return false;
+	}
+
+	return ItemState->PlayingUntil_RootTicks > Sequencer->GetGlobalTime().Time.FloorToFrame().Value;
+}
+
+
+void FSequencerPlaylistItemPlayer_Sequence::ClearItemStates()
+{
+	ItemStates.Empty();
+}
+
+
 UMovieSceneSubTrack* FSequencerPlaylistItemPlayer_Sequence::GetOrCreateWorkingTrack(USequencerPlaylistItem* Item)
 {
 	USequencerPlaylistItem_Sequence* SequenceItem = CastChecked<USequencerPlaylistItem_Sequence>(Item);
@@ -245,7 +294,7 @@ UMovieSceneSubTrack* FSequencerPlaylistItemPlayer_Sequence::GetOrCreateWorkingTr
 		return nullptr;
 	}
 
-	FItemState& ItemState = GetOrCreateItemState(Item);
+	FItemState& ItemState = ItemStates.FindOrAdd(Item);
 
 	ULevelSequence* RootSequence = Cast<ULevelSequence>(Sequencer->GetRootMovieSceneSequence());
 	UMovieScene* RootScene = RootSequence->GetMovieScene();
@@ -293,13 +342,6 @@ UMovieSceneSubTrack* FSequencerPlaylistItemPlayer_Sequence::GetOrCreateWorkingTr
 }
 
 
-FSequencerPlaylistItemPlayer_Sequence::FItemState& FSequencerPlaylistItemPlayer_Sequence::GetOrCreateItemState(USequencerPlaylistItem* Item)
-{
-	check(Item);
-	return ItemStates.FindOrAdd(Item);
-}
-
-
 void FSequencerPlaylistItemPlayer_Sequence::EndSection(UMovieSceneSection* Section)
 {
 	if (!ensure(Section))
@@ -326,6 +368,17 @@ void FSequencerPlaylistItemPlayer_Sequence::EndSection(UMovieSceneSection* Secti
 	{
 		Section->SetEndFrame(TRangeBound<FFrameNumber>::Exclusive(SectionNow.FloorToFrame()));
 
+		// Remove degenerate sections. This can happen if we reset then stop a held item while paused.
+		if (Section->GetRange().IsEmpty())
+		{
+			if (UMovieSceneTrack* SectionTrack = Section->GetTypedOuter<UMovieSceneTrack>())
+			{
+				SectionTrack->Modify();
+				SectionTrack->RemoveSection(*Section);
+				return;
+			}
+		}
+
 		// Workaround for not being able to interrupt first playthrough if bCanLoop == true.
 		if (UMovieSceneSubSection* SubSection = Cast<UMovieSceneSubSection>(Section))
 		{
diff --git a/Engine/Plugins/VirtualProduction/SequencerPlaylists/Source/SequencerPlaylists/Private/SequencerPlaylistItem_Sequence.h b/Engine/Plugins/VirtualProduction/SequencerPlaylists/Source/SequencerPlaylists/Private/SequencerPlaylistItem_Sequence.h
--- a/Engine/Plugins/VirtualProduction/SequencerPlaylists/Source/SequencerPlaylists/Private/SequencerPlaylistItem_Sequence.h
+++ b/Engine/Plugins/VirtualProduction/SequencerPlaylists/Source/SequencerPlaylists/Private/SequencerPlaylistItem_Sequence.h
@@ -34,21 +34,26 @@ class FSequencerPlaylistItemPlayer_Sequence : public ISequencerPlaylistItemPlaye
 		TWeakObjectPtr<UMovieSceneSubTrack> WeakTrack;
 		TWeakObjectPtr<UMovieSceneSubSection> WeakHoldSection;
 		TArray<TWeakObjectPtr<UMovieSceneSubSection>> WeakPlaySections;
+		int32 PlayingUntil_RootTicks = TNumericLimits<int32>::Min();
 	};
 
 public:
 	FSequencerPlaylistItemPlayer_Sequence(TSharedRef<ISequencer> Sequencer);
+	~FSequencerPlaylistItemPlayer_Sequence() override;
 
 	//~ Begin ISequencerPlaylistItemPlayer
 	bool Play(USequencerPlaylistItem* Item) override;
 	bool Stop(USequencerPlaylistItem* Item) override;
 	bool AddHold(USequencerPlaylistItem* Item) override;
 	bool Reset(USequencerPlaylistItem* Item) override;
+
+	bool IsPlaying(USequencerPlaylistItem* Item) const override;
 	//~ End ISequencerPlaylistItemPlayer
 
 private:
+	void ClearItemStates();
+
 	UMovieSceneSubTrack* GetOrCreateWorkingTrack(USequencerPlaylistItem* Item);
-	FItemState& GetOrCreateItemState(USequencerPlaylistItem* Item);
 	void EndSection(UMovieSceneSection* Section);
 
 private:
diff --git a/Engine/Plugins/VirtualProduction/SequencerPlaylists/Source/SequencerPlaylists/Private/SequencerPlaylistPlayer.cpp b/Engine/Plugins/VirtualProduction/SequencerPlaylists/Source/SequencerPlaylists/Private/SequencerPlaylistPlayer.cpp
--- a/Engine/Plugins/VirtualProduction/SequencerPlaylists/Source/SequencerPlaylists/Private/SequencerPlaylistPlayer.cpp
+++ b/Engine/Plugins/VirtualProduction/SequencerPlaylists/Source/SequencerPlaylists/Private/SequencerPlaylistPlayer.cpp
@@ -94,6 +94,18 @@ bool USequencerPlaylistPlayer::ResetItem(USequencerPlaylistItem* Item)
 	return GetCheckedItemPlayer(Item)->Reset(Item);
 }
 
+
+bool USequencerPlaylistPlayer::IsPlaying(USequencerPlaylistItem* Item)
+{
+	if (!Item)
+	{
+		return false;
+	}
+
+	return GetCheckedItemPlayer(Item)->IsPlaying(Item);
+}
+
+
 namespace UE::Private::PlaylistPlayer
 {
 
@@ -150,8 +162,8 @@ void StopPlaybackAndAdjustTime(TSharedPtr<ISequencer>& Sequencer)
 {
 	check(Sequencer);
 
-	Sequencer->SetPlaybackStatus(EMovieScenePlayerStatus::Stopped);
-	UMovieSceneSequence* Sequence  = Sequencer->GetRootMovieSceneSequence();
+	Sequencer->Pause();
+	UMovieSceneSequence* Sequence = Sequencer->GetRootMovieSceneSequence();
 	if (Sequence)
 	{
 		UMovieScene* MovieScene = Sequence->GetMovieScene();
@@ -165,11 +177,15 @@ void StopPlaybackAndAdjustTime(TSharedPtr<ISequencer>& Sequencer)
 	}
 }
 
-}
+} // namespace UE::Private::PlaylistPlayer
 
 void USequencerPlaylistPlayer::Tick(float DeltaTime)
 {
-	TSharedPtr<ISequencer> Sequencer = GetSequencer();
+	TSharedPtr<ISequencer> Sequencer = WeakSequencer.Pin();
+	if (!Sequencer)
+	{
+		return;
+	}
 
 	if (Sequencer->GetPlaybackStatus() == EMovieScenePlayerStatus::Paused)
 	{
@@ -182,6 +198,7 @@ void USequencerPlaylistPlayer::Tick(float DeltaTime)
 		UE::Private::PlaylistPlayer::StopPlaybackAndAdjustTime(Sequencer);
 		return;
 	}
+
 	// Handle a tick
 	UE::Private::PlaylistPlayer::AdjustMovieSceneRangeForPlay(Sequencer);
 }
@@ -240,15 +257,6 @@ bool USequencerPlaylistPlayer::StopAll()
 		return false;
 	}
 
-	UTakeRecorder* Recorder = UTakeRecorder::GetActiveRecorder();
-
-	TSharedPtr<ISequencer> Sequencer = GetSequencer();
-	const bool bInRecorder = Recorder && Recorder->GetState() != ETakeRecorderState::Stopped;
-	if (!bInRecorder && Sequencer->GetPlaybackStatus() == EMovieScenePlayerStatus::Playing)
-	{
-		UE::Private::PlaylistPlayer::StopPlaybackAndAdjustTime(Sequencer);
-	}
-
 	bool bAnyChange = false;
 
 	FScopedTransaction Transaction(LOCTEXT("StopAllTransaction", "Stop playback of all items"));
@@ -257,6 +265,15 @@ bool USequencerPlaylistPlayer::StopAll()
 		bAnyChange |= GetCheckedItemPlayer(Item)->Stop(Item);
 	}
 
+	TSharedPtr<ISequencer> Sequencer = GetSequencer();
+
+	UTakeRecorder* Recorder = UTakeRecorder::GetActiveRecorder();
+	const bool bInRecorder = Recorder && Recorder->GetState() != ETakeRecorderState::Stopped;
+	if (!bInRecorder && Sequencer->GetPlaybackStatus() == EMovieScenePlayerStatus::Playing)
+	{
+		UE::Private::PlaylistPlayer::StopPlaybackAndAdjustTime(Sequencer);
+	}
+
 	if (bAnyChange)
 	{
 		Sequencer->NotifyMovieSceneDataChanged(EMovieSceneDataChangeType::RefreshAllImmediately);
diff --git a/Engine/Plugins/VirtualProduction/SequencerPlaylists/Source/SequencerPlaylists/Private/SequencerPlaylistsWidgets.cpp b/Engine/Plugins/VirtualProduction/SequencerPlaylists/Source/SequencerPlaylists/Private/SequencerPlaylistsWidgets.cpp
--- a/Engine/Plugins/VirtualProduction/SequencerPlaylists/Source/SequencerPlaylists/Private/SequencerPlaylistsWidgets.cpp
+++ b/Engine/Plugins/VirtualProduction/SequencerPlaylists/Source/SequencerPlaylists/Private/SequencerPlaylistsWidgets.cpp
@@ -328,6 +328,18 @@ TSharedRef<SWidget> SSequencerPlaylistPanel::Construct_ItemListView()
 			{
 				return SNew(SSequencerPlaylistItemWidget, InData, OwnerTableView)
 					.TriggerMode(TAttribute<bool>::Create(TAttribute<bool>::FGetter::CreateSP(This, &SSequencerPlaylistPanel::InTriggerMode)))
+					.IsPlaying_Lambda([This, WeakItem = InData->WeakItem]() {
+						USequencerPlaylistItem* Item = WeakItem.Get();
+						USequencerPlaylistPlayer* Player = This->WeakPlayer.Get();
+						if (Item && Player)
+						{
+							return Player->IsPlaying(Item);
+						}
+						else
+						{
+							return false;
+						}
+					})
 					.OnPlayClicked(This, &SSequencerPlaylistPanel::HandleClicked_Item_Play)
 					.OnStopClicked(This, &SSequencerPlaylistPanel::HandleClicked_Item_Stop)
 					.OnResetClicked(This, &SSequencerPlaylistPanel::HandleClicked_Item_Reset)
@@ -869,6 +881,7 @@ void SSequencerPlaylistItemWidget::Construct(const FArguments& InArgs, TSharedPt
 	RowData = InRowData;
 
 	TriggerMode = InArgs._TriggerMode;
+	IsPlaying = InArgs._IsPlaying;
 
 	PlayClickedDelegate = InArgs._OnPlayClicked;
 	StopClickedDelegate = InArgs._OnStopClicked;
@@ -990,7 +1003,6 @@ TSharedRef<SWidget> SSequencerPlaylistItemWidget::GenerateWidgetForColumn(const
 	if (ColumnName == SSequencerPlaylistPanel::ColumnName_HoverTransport)
 	{
 		return SNew(SVerticalBox)
-			.Visibility(TAttribute<EVisibility>::Create(TAttribute<EVisibility>::FGetter::CreateSP(SharedThis(this), &SSequencerPlaylistItemWidget::GetHoverTransportCellVisibility)))
 			+ SVerticalBox::Slot()
 			[
 				SNew(SButton)
@@ -998,6 +1010,7 @@ TSharedRef<SWidget> SSequencerPlaylistItemWidget::GenerateWidgetForColumn(const
 				.ContentPadding(FMargin(1.0f, 3.0f, 0.0f, 0.0f))
 				.HAlign(HAlign_Center)
 				.VAlign(VAlign_Center)
+				.Visibility_Lambda([This]() { return (This->IsHovered() && !This->InTriggerMode()) ? EVisibility::Visible : EVisibility::Hidden; })
 				.OnClicked_Lambda([This]() { return This->PlayClickedDelegate.Execute(This); })
 				.ToolTipText(PlayItemTooltipText)
 				[
@@ -1013,6 +1026,8 @@ TSharedRef<SWidget> SSequencerPlaylistItemWidget::GenerateWidgetForColumn(const
 				.ContentPadding(FMargin(0.0f, 0.0f, 0.0f, 1.0f))
 				.HAlign(HAlign_Center)
 				.VAlign(VAlign_Center)
+				.Visibility_Lambda([This]() { return (This->IsPlaying.Get() || (This->IsHovered() && !This->InTriggerMode())) ? EVisibility::Visible : EVisibility::Hidden; })
+				.ForegroundColor_Lambda([This]() { return This->IsPlaying.Get() ? FStyleColors::AccentRed : FSlateColor::UseStyle(); })
 				.OnClicked_Lambda([This]() { return This->StopClickedDelegate.Execute(This); })
 				.ToolTipText(StopItemTooltipText)
 				[
@@ -1028,6 +1043,7 @@ TSharedRef<SWidget> SSequencerPlaylistItemWidget::GenerateWidgetForColumn(const
 				.ContentPadding(FMargin(1.0f, 0.0f, 0.0f, 3.0f))
 				.HAlign(HAlign_Center)
 				.VAlign(VAlign_Center)
+				.Visibility_Lambda([This]() { return (This->IsHovered() && !This->InTriggerMode()) ? EVisibility::Visible : EVisibility::Hidden; })
 				.OnClicked_Lambda([This]() { return This->ResetClickedDelegate.Execute(This); })
 				.ToolTipText(ResetItemTooltipText)
 				[
@@ -1165,17 +1181,6 @@ EVisibility SSequencerPlaylistItemWidget::GetTriggerModeTransportVisibility() co
 }
 
 
-EVisibility SSequencerPlaylistItemWidget::GetHoverTransportCellVisibility() const
-{
-	if (IsHovered() && !InTriggerMode())
-	{
-		return EVisibility::Visible;
-	}
-
-	return EVisibility::Hidden;
-}
-
-
 TSharedRef<SWidget> SSequencerPlaylistItemWidget::EnsureSelectedAndBuildContextMenu()
 {
 	TArray<UObject*> SelectedItems;
diff --git a/Engine/Plugins/VirtualProduction/SequencerPlaylists/Source/SequencerPlaylists/Private/SequencerPlaylistsWidgets.h b/Engine/Plugins/VirtualProduction/SequencerPlaylists/Source/SequencerPlaylists/Private/SequencerPlaylistsWidgets.h
--- a/Engine/Plugins/VirtualProduction/SequencerPlaylists/Source/SequencerPlaylists/Private/SequencerPlaylistsWidgets.h
+++ b/Engine/Plugins/VirtualProduction/SequencerPlaylists/Source/SequencerPlaylists/Private/SequencerPlaylistsWidgets.h
@@ -131,6 +131,7 @@ class SSequencerPlaylistItemWidget : public SMultiColumnTableRow<TSharedPtr<FSeq
 
 	SLATE_BEGIN_ARGS(SSequencerPlaylistItemWidget) {}
 		SLATE_ATTRIBUTE(bool, TriggerMode)
+		SLATE_ATTRIBUTE(bool, IsPlaying)
 
 		SLATE_EVENT(FOnClickedSequencerPlaylistItem, OnPlayClicked)
 		SLATE_EVENT(FOnClickedSequencerPlaylistItem, OnStopClicked)
@@ -166,7 +167,6 @@ private:
 	bool InTriggerMode() const { return TriggerMode.Get(); }
 	bool IsRowContentEnabled() const;
 	EVisibility GetTriggerModeTransportVisibility() const;
-	EVisibility GetHoverTransportCellVisibility() const;
 
 	TSharedRef<SWidget> EnsureSelectedAndBuildContextMenu();
 	TSharedRef<SWidget> BuildContextMenu(const TArray<UObject*>& SelectedItems);
@@ -176,6 +176,7 @@ private:
 	TSharedPtr<SMenuAnchor> DetailsAnchor;
 
 	TAttribute<bool> TriggerMode;
+	TAttribute<bool> IsPlaying;
 
 	FOnClickedSequencerPlaylistItem PlayClickedDelegate;
 	FOnClickedSequencerPlaylistItem StopClickedDelegate;
diff --git a/Engine/Plugins/VirtualProduction/SequencerPlaylists/Source/SequencerPlaylists/Public/ISequencerPlaylistsModule.h b/Engine/Plugins/VirtualProduction/SequencerPlaylists/Source/SequencerPlaylists/Public/ISequencerPlaylistsModule.h
--- a/Engine/Plugins/VirtualProduction/SequencerPlaylists/Source/SequencerPlaylists/Public/ISequencerPlaylistsModule.h
+++ b/Engine/Plugins/VirtualProduction/SequencerPlaylists/Source/SequencerPlaylists/Public/ISequencerPlaylistsModule.h
@@ -43,6 +43,9 @@ public:
 	 * @return True if the current sequence was modified, otherwise false.
 	 */
 	virtual bool Reset(USequencerPlaylistItem* Item) = 0;
+
+	/** Returns whether we're currently playing the specified item. */
+	virtual bool IsPlaying(USequencerPlaylistItem* Item) const = 0;
 };
 
 
diff --git a/Engine/Plugins/VirtualProduction/SequencerPlaylists/Source/SequencerPlaylists/Public/SequencerPlaylistPlayer.h b/Engine/Plugins/VirtualProduction/SequencerPlaylists/Source/SequencerPlaylists/Public/SequencerPlaylistPlayer.h
--- a/Engine/Plugins/VirtualProduction/SequencerPlaylists/Source/SequencerPlaylists/Public/SequencerPlaylistPlayer.h
+++ b/Engine/Plugins/VirtualProduction/SequencerPlaylists/Source/SequencerPlaylists/Public/SequencerPlaylistPlayer.h
@@ -51,6 +51,9 @@ public:
 	UFUNCTION(BlueprintCallable, Category="SequencerPlaylists")
 	bool ResetItem(USequencerPlaylistItem* Item);
 
+	UFUNCTION(BlueprintPure, Category="SequencerPlaylists")
+	bool IsPlaying(USequencerPlaylistItem* Item);
+
 	UFUNCTION(BlueprintCallable, Category="SequencerPlaylists")
 	bool PlayAll();
 
