Sequencer Playlists: Finish hooking up item "Mute" feature.

diff --git a/Engine/Plugins/VirtualProduction/SequencerPlaylists/Source/SequencerPlaylists/Private/SequencerPlaylistPlayer.cpp b/Engine/Plugins/VirtualProduction/SequencerPlaylists/Source/SequencerPlaylists/Private/SequencerPlaylistPlayer.cpp
--- a/Engine/Plugins/VirtualProduction/SequencerPlaylists/Source/SequencerPlaylists/Private/SequencerPlaylistPlayer.cpp
+++ b/Engine/Plugins/VirtualProduction/SequencerPlaylists/Source/SequencerPlaylists/Private/SequencerPlaylistPlayer.cpp
@@ -238,6 +238,11 @@ bool USequencerPlaylistPlayer::PlayAll()
 	FScopedTransaction Transaction(LOCTEXT("PlayAllTransaction", "Trigger playback of all items"));
 	for (USequencerPlaylistItem* Item : Playlist->Items)
 	{
+		if (Item->bMute)
+		{
+			continue;
+		}
+
 		bAnyChange |= GetCheckedItemPlayer(Item)->Play(Item);
 	}
 
diff --git a/Engine/Plugins/VirtualProduction/SequencerPlaylists/Source/SequencerPlaylists/Private/SequencerPlaylistsWidgets.cpp b/Engine/Plugins/VirtualProduction/SequencerPlaylists/Source/SequencerPlaylists/Private/SequencerPlaylistsWidgets.cpp
--- a/Engine/Plugins/VirtualProduction/SequencerPlaylists/Source/SequencerPlaylists/Private/SequencerPlaylistsWidgets.cpp
+++ b/Engine/Plugins/VirtualProduction/SequencerPlaylists/Source/SequencerPlaylists/Private/SequencerPlaylistsWidgets.cpp
@@ -919,7 +919,7 @@ void SSequencerPlaylistItemWidget::ConstructChildren(ETableViewMode::Type InOwne
 		+ SOverlay::Slot()
 		[
 			SNew(SImage)
-			.Visibility(this, &SSequencerPlaylistItemWidget::GetTriggerModeTransportVisibility)
+			.Visibility(this, &SSequencerPlaylistItemWidget::GetRowDimmingVisibility)
 			.Image(FSequencerPlaylistsStyle::Get().GetBrush("SequencerPlaylists.Item.Dim"))
 		]
 		+ SOverlay::Slot()
@@ -1164,9 +1164,19 @@ FReply SSequencerPlaylistItemWidget::HandleDragDetected(const FGeometry& MyGeome
 }
 
 
-bool SSequencerPlaylistItemWidget::IsRowContentEnabled() const
+EVisibility SSequencerPlaylistItemWidget::GetRowDimmingVisibility() const
 {
-	return GetTriggerModeTransportVisibility() == EVisibility::Hidden;
+	if (GetTriggerModeTransportVisibility() == EVisibility::Visible)
+	{
+		return EVisibility::Visible;
+	}
+
+	if (USequencerPlaylistItem* Item = RowData->WeakItem.Get())
+	{
+		return Item->bMute ? EVisibility::SelfHitTestInvisible : EVisibility::Hidden;
+	}
+
+	return EVisibility::Hidden;
 }
 
 
diff --git a/Engine/Plugins/VirtualProduction/SequencerPlaylists/Source/SequencerPlaylists/Private/SequencerPlaylistsWidgets.h b/Engine/Plugins/VirtualProduction/SequencerPlaylists/Source/SequencerPlaylists/Private/SequencerPlaylistsWidgets.h
--- a/Engine/Plugins/VirtualProduction/SequencerPlaylists/Source/SequencerPlaylists/Private/SequencerPlaylistsWidgets.h
+++ b/Engine/Plugins/VirtualProduction/SequencerPlaylists/Source/SequencerPlaylists/Private/SequencerPlaylistsWidgets.h
@@ -165,7 +165,8 @@ private:
 	FReply HandleDragDetected(const FGeometry& MyGeometry, const FPointerEvent& MouseEvent);
 
 	bool InTriggerMode() const { return TriggerMode.Get(); }
-	bool IsRowContentEnabled() const;
+
+	EVisibility GetRowDimmingVisibility() const;
 	EVisibility GetTriggerModeTransportVisibility() const;
 
 	TSharedRef<SWidget> EnsureSelectedAndBuildContextMenu();
diff --git a/Engine/Plugins/VirtualProduction/SequencerPlaylists/Source/SequencerPlaylists/Public/SequencerPlaylistItem.h b/Engine/Plugins/VirtualProduction/SequencerPlaylists/Source/SequencerPlaylists/Public/SequencerPlaylistItem.h
--- a/Engine/Plugins/VirtualProduction/SequencerPlaylists/Source/SequencerPlaylists/Public/SequencerPlaylistItem.h
+++ b/Engine/Plugins/VirtualProduction/SequencerPlaylists/Source/SequencerPlaylists/Public/SequencerPlaylistItem.h
@@ -38,4 +38,7 @@ public:
 
 	UPROPERTY(EditAnywhere, BlueprintReadWrite, Category="SequencerPlaylists", meta=(NoResetToDefault))
 	float PlaybackSpeed = 1.0f;
+
+	UPROPERTY(EditAnywhere, BlueprintReadWrite, Category="SequencerPlaylists", meta=(NoResetToDefault))
+	bool bMute;
 };
