Sequencer Playlists: Resolve missing final UI elements. Also implements "Save" (as opposed to "Save As") behavior, and fixes a bug where muted items would still add holds.

diff --git a/Engine/Plugins/VirtualProduction/SequencerPlaylists/Source/SequencerPlaylists/Private/SequencerPlaylistPlayer.cpp b/Engine/Plugins/VirtualProduction/SequencerPlaylists/Source/SequencerPlaylists/Private/SequencerPlaylistPlayer.cpp
--- a/Engine/Plugins/VirtualProduction/SequencerPlaylists/Source/SequencerPlaylists/Private/SequencerPlaylistPlayer.cpp
+++ b/Engine/Plugins/VirtualProduction/SequencerPlaylists/Source/SequencerPlaylists/Private/SequencerPlaylistPlayer.cpp
@@ -302,7 +302,15 @@ bool USequencerPlaylistPlayer::ResetAll()
 	FScopedTransaction Transaction(LOCTEXT("ResetAllTransaction", "Reset playback of all items"));
 	for (USequencerPlaylistItem* Item : Playlist->Items)
 	{
-		bAnyChange |= GetCheckedItemPlayer(Item)->Reset(Item);
+		if (Item->bMute)
+		{
+			// Stop muted items, but don't add holds for them.
+			bAnyChange |= GetCheckedItemPlayer(Item)->Stop(Item);
+		}
+		else
+		{
+			bAnyChange |= GetCheckedItemPlayer(Item)->Reset(Item);
+		}
 	}
 
 	if (bAnyChange)
@@ -412,7 +420,7 @@ void USequencerPlaylistPlayer::OnTakeRecorderStarted(UTakeRecorder* InRecorder)
 
 		for (USequencerPlaylistItem* Item : Playlist->Items)
 		{
-			if (Item->bHoldAtFirstFrame)
+			if (Item->bHoldAtFirstFrame && !Item->bMute)
 			{
 				GetCheckedItemPlayer(Item)->AddHold(Item);
 			}
diff --git a/Engine/Plugins/VirtualProduction/SequencerPlaylists/Source/SequencerPlaylists/Private/SequencerPlaylistsModule.cpp b/Engine/Plugins/VirtualProduction/SequencerPlaylists/Source/SequencerPlaylists/Private/SequencerPlaylistsModule.cpp
--- a/Engine/Plugins/VirtualProduction/SequencerPlaylists/Source/SequencerPlaylists/Private/SequencerPlaylistsModule.cpp
+++ b/Engine/Plugins/VirtualProduction/SequencerPlaylists/Source/SequencerPlaylists/Private/SequencerPlaylistsModule.cpp
@@ -52,8 +52,8 @@ void FSequencerPlaylistsModule::StartupModule()
 	UToolMenus::RegisterStartupCallback(FSimpleMulticastDelegate::FDelegate::CreateRaw(this, &FSequencerPlaylistsModule::RegisterMenus));
 
 	FGlobalTabmanager::Get()->RegisterNomadTabSpawner(SequencerPlaylistsTabName, FOnSpawnTab::CreateRaw(this, &FSequencerPlaylistsModule::OnSpawnPluginTab))
-		.SetDisplayName(LOCTEXT("SequencerPlaylistsTabTitle", "Playlists"))
-		.SetTooltipText(LOCTEXT("SequencerPlaylistsTabTooltipText", "Open the Sequencer Playlists tab."))
+		.SetDisplayName(LOCTEXT("TabSpawnerDisplayName", "Playlists"))
+		.SetTooltipText(LOCTEXT("TabSpawnerTooltipText", "Open the Sequencer Playlists tab."))
 		.SetGroup(WorkspaceMenu::GetMenuStructure().GetLevelEditorVirtualProductionCategory())
 		.SetIcon(FSlateIcon(FSequencerPlaylistsStyle::GetStyleSetName(), "SequencerPlaylists.TabIcon"));
 }
@@ -103,6 +103,7 @@ TSharedRef<SDockTab> FSequencerPlaylistsModule::OnSpawnPluginTab(const FSpawnTab
 {
 	return SNew(SDockTab)
 		.TabRole(ETabRole::NomadTab)
+		.Label(LOCTEXT("TabLabel", "Playlist"))
 		[
 			SNew(SSequencerPlaylistPanel, DefaultPlayer.Get())
 		];
diff --git a/Engine/Plugins/VirtualProduction/SequencerPlaylists/Source/SequencerPlaylists/Private/SequencerPlaylistsStyle.cpp b/Engine/Plugins/VirtualProduction/SequencerPlaylists/Source/SequencerPlaylists/Private/SequencerPlaylistsStyle.cpp
--- a/Engine/Plugins/VirtualProduction/SequencerPlaylists/Source/SequencerPlaylists/Private/SequencerPlaylistsStyle.cpp
+++ b/Engine/Plugins/VirtualProduction/SequencerPlaylists/Source/SequencerPlaylists/Private/SequencerPlaylistsStyle.cpp
@@ -8,6 +8,7 @@
 #include "Styling/SlateStyleMacros.h"
 #include "Styling/SlateStyleRegistry.h"
 #include "Styling/StyleColors.h"
+#include "Styling/ToolBarStyle.h"
 
 
 #define RootToContentDir Style->RootToContentDir
@@ -58,23 +59,64 @@ TSharedRef<FSlateStyleSet> FSequencerPlaylistsStyle::Create()
 	static const FVector2D Icon20x20(20.0f, 20.0f);
 	static const FVector2D Icon24x24(24.0f, 24.0f);
 
+	const FString PluginContentRoot = IPluginManager::Get().FindPlugin("SequencerPlaylists")->GetBaseDir() / TEXT("Resources");
+	const FString EngineSlateDir = FPaths::EngineContentDir() / TEXT("Slate");
+	const FString EngineEditorSlateDir = FPaths::EngineContentDir() / TEXT("Editor") / TEXT("Slate");
+
 	TSharedRef<FSlateStyleSet> Style = MakeShareable(new FSlateStyleSet("SequencerPlaylistsStyle"));
-	Style->SetContentRoot(IPluginManager::Get().FindPlugin("SequencerPlaylists")->GetBaseDir() / TEXT("Resources"));
+	Style->SetContentRoot(PluginContentRoot);
 
-	Style->Set("SequencerPlaylists.TabIcon", new IMAGE_BRUSH_SVG("PlaylistIcon", Icon16x16));
-	Style->Set("SequencerPlaylists.OpenPluginWindow", new IMAGE_BRUSH_SVG("PlaceholderButtonIcon", Icon20x20));
+	Style->Set("SequencerPlaylists.TabIcon", new IMAGE_BRUSH_SVG("Playlist", Icon16x16));
 	Style->Set("SequencerPlaylists.Panel.Background", new FSlateColorBrush(FStyleColors::Panel));
 
 	FLinearColor DimColor(FSlateColor(EStyleColor::Background).GetSpecifiedColor());
 	DimColor.A = 0.5f;
 	Style->Set("SequencerPlaylists.Item.Dim", new FSlateColorBrush(DimColor));
 
-	Style->Set("SequencerPlaylists.NewPlaylist.Background", new IMAGE_BRUSH_SVG("PlaylistIcon", Icon20x20));
-	// NewPlaylist.Overlay, SavePlaylistAs, OpenPlaylist are registered under a different root below.
+	Style->Set("SequencerPlaylists.NewPlaylist", new IMAGE_BRUSH_SVG("PlaylistNew", Icon16x16));
+	// SavePlaylist is registered under a different root below.
+	Style->Set("SequencerPlaylists.OpenPlaylist", new IMAGE_BRUSH_SVG("PlaylistLoad", Icon16x16));
 
-	FButtonStyle TransportButtonStyle = FButtonStyle()
-		.SetNormalForeground(FStyleColors::Foreground)
-		.SetDisabledForeground(FStyleColors::Foreground);
+	Style->Set("SequencerPlaylists.PlayMode", new IMAGE_BRUSH_SVG("PlaybackAuto", Icon16x16));
+
+	Style->Set("SequencerPlaylists.HoldFrame", new IMAGE_BRUSH_SVG("HoldFrame", Icon16x16));
+
+	Style->Set("SequencerPlaylists.Loop.Disabled", new IMAGE_BRUSH_SVG("LoopNot", Icon16x16));
+	Style->Set("SequencerPlaylists.Loop.Finite", new IMAGE_BRUSH_SVG("LoopNumber", Icon16x16));
+
+	// Main toolbar
+	{
+		const FToolBarStyle& SlimToolbarStyle = FAppStyle::Get().GetWidgetStyle<FToolBarStyle>("SlimToolBar");
+		FToolBarStyle MainToolbarStyle = FToolBarStyle(SlimToolbarStyle)
+			.SetIconSize(Icon16x16)
+			.SetComboButtonPadding(FMargin(2.0f, 0.0f));
+
+		FButtonStyle MainToolbarButtonStyle = FButtonStyle(SlimToolbarStyle.ButtonStyle)
+			.SetNormalPadding(FMargin(4.f, 4.f, 4.f, 4.f))
+			.SetPressedPadding(FMargin(4.f, 5.f, 4.f, 3.f));
+		MainToolbarStyle.SetButtonStyle(MainToolbarButtonStyle);
+
+		FComboButtonStyle MainToolbarComboButtonStyle = FComboButtonStyle(SlimToolbarStyle.ComboButtonStyle)
+			.SetDownArrowPadding(FMargin(1.f, 0.f, 0.f, 0.f))
+			.SetButtonStyle(MainToolbarButtonStyle);
+		MainToolbarStyle.SetComboButtonStyle(MainToolbarComboButtonStyle);
+
+		{
+			// Borrowed asset; change the root temporarily
+			Style->SetContentRoot(EngineSlateDir);
+			ON_SCOPE_EXIT{ Style->SetContentRoot(PluginContentRoot); };
+
+			FComboButtonStyle MainToolbarSettingsComboButtonStyle = FComboButtonStyle(SlimToolbarStyle.SettingsComboButton)
+				.SetDownArrowImage(IMAGE_BRUSH_SVG("Starship/Common/ellipsis-vertical-narrow", FVector2D(4, 16)));
+			MainToolbarStyle.SetSettingsComboButtonStyle(MainToolbarSettingsComboButtonStyle);
+		}
+
+		FCheckBoxStyle MainToolbarToggleButtonStyle = FCheckBoxStyle(SlimToolbarStyle.ToggleButton)
+			.SetPadding(FMargin(4.f, 4.f, 4.f, 4.f));
+		MainToolbarStyle.SetToggleButtonStyle(MainToolbarToggleButtonStyle);
+
+		Style->Set("SequencerPlaylists.MainToolbar", MainToolbarStyle);
+	}
 
 	const FLinearColor PlayColor = FStyleColors::AccentGreen.GetSpecifiedColor();
 	const FLinearColor StopColor = FStyleColors::AccentRed.GetSpecifiedColor();
@@ -84,14 +126,34 @@ TSharedRef<FSlateStyleSet> FSequencerPlaylistsStyle::Create()
 	const FLinearColor ResetPressedColor = MakeColorVariation(ResetColor, EColorVariation::Dimmed);
 
 	const float CornerRadius = 4.0f;
+	const FVector4 AllCorners(CornerRadius, CornerRadius, CornerRadius, CornerRadius);
 	const FVector4 LeftCorners(CornerRadius, 0.0f, 0.0f, CornerRadius);
 	const FVector4 MiddleCorners(0.0f, 0.0f, 0.0f, 0.0f);
 	const FVector4 RightCorners(0.0f, CornerRadius, CornerRadius, 0.0f);
 
+	// Table rows
+	{
+		const FTableRowStyle& NormalTableRowStyle = FEditorStyle::Get().GetWidgetStyle<FTableRowStyle>("TableView.Row");
+		FTableRowStyle PlaylistItemRowStyle = FTableRowStyle(NormalTableRowStyle)
+			.SetEvenRowBackgroundBrush(FSlateNoResource())
+			.SetOddRowBackgroundBrush(FSlateNoResource())
+			.SetEvenRowBackgroundHoveredBrush(FSlateNoResource())
+			.SetOddRowBackgroundHoveredBrush(FSlateNoResource());
+		Style->Set("SequencerPlaylists.ItemRow", PlaylistItemRowStyle);
+
+		Style->Set("SequencerPlaylists.ItemRow.BackgroundOuter", new FSlateColorBrush(FStyleColors::Recessed));
+		Style->Set("SequencerPlaylists.ItemRow.BackgroundInner", new FSlateRoundedBoxBrush(FStyleColors::Panel, AllCorners));
+	}
+
+
 	const FSlateRoundedBoxBrush LeftBrush(FStyleColors::Dropdown, LeftCorners);
 	const FSlateRoundedBoxBrush MiddleBrush(FStyleColors::Dropdown, MiddleCorners);
 	const FSlateRoundedBoxBrush RightBrush(FStyleColors::Dropdown, RightCorners);
 
+	FButtonStyle TransportButtonStyle = FButtonStyle()
+		.SetNormalForeground(FStyleColors::Foreground)
+		.SetDisabledForeground(FStyleColors::Foreground);
+
 	const FButtonStyle LeftTransportButton = FButtonStyle(TransportButtonStyle)
 		.SetNormal(LeftBrush)
 		.SetHovered(LeftBrush)
@@ -166,35 +228,31 @@ TSharedRef<FSlateStyleSet> FSequencerPlaylistsStyle::Create()
 		.SetForegroundColor(FSlateColor::UseForeground());
 
 	Style->Set("SequencerPlaylists.EditableTextBox", EditableTextStyle);
-	Style->Set("SequencerPlaylists.TitleFont", FCoreStyle::GetDefaultFontStyle("Bold", 18));
+	Style->Set("SequencerPlaylists.TitleFont", FCoreStyle::GetDefaultFontStyle("Regular", 13));
 	Style->Set("SequencerPlaylists.DescriptionFont", FCoreStyle::GetDefaultFontStyle("Regular", 8));
 
 	{
 		// These assets live in engine slate folders; change the root temporarily
-		const FString PreviousContentRoot = Style->GetContentRootDir();
-		ON_SCOPE_EXIT{ Style->SetContentRoot(PreviousContentRoot); };
+		ON_SCOPE_EXIT{ Style->SetContentRoot(PluginContentRoot); };
 
 		// Engine/Content/Slate/...
 		{
-			const FString EngineSlateDir = FPaths::EngineContentDir() / TEXT("Slate");
 			Style->SetContentRoot(EngineSlateDir);
-
 			Style->Set("SequencerPlaylists.Ellipsis", new IMAGE_BRUSH_SVG("Starship/Common/ellipsis-vertical-narrow", FVector2D(6, 24)));
 		}
 
 		// Engine/Content/Editor/Slate/...
 		{
-			const FString EngineEditorSlateDir = FPaths::EngineContentDir() / TEXT("Editor") / TEXT("Slate");
 			Style->SetContentRoot(EngineEditorSlateDir);
 
-			Style->Set("SequencerPlaylists.NewPlaylist.Overlay", new IMAGE_BRUSH_SVG("Starship/MainToolbar/ToolBadgePlus", Icon20x20, FStyleColors::AccentGreen));
-			Style->Set("SequencerPlaylists.SavePlaylistAs", new IMAGE_BRUSH_SVG("Starship/Common/SaveCurrentAs", Icon16x16));
-			Style->Set("SequencerPlaylists.OpenPlaylist", new IMAGE_BRUSH_SVG("Starship/Common/OpenAsset", Icon16x16));
+			Style->Set("SequencerPlaylists.SavePlaylist", new IMAGE_BRUSH_SVG("Starship/Common/SaveCurrent", Icon16x16));
 
 			Style->Set("SequencerPlaylists.Play", new IMAGE_BRUSH_SVG("Starship/Common/play", Icon20x20));
 			Style->Set("SequencerPlaylists.Play.Small", new IMAGE_BRUSH_SVG("Starship/Common/play", Icon12x12));
 			Style->Set("SequencerPlaylists.Stop", new IMAGE_BRUSH_SVG("Starship/MainToolbar/stop", Icon20x20));
 			Style->Set("SequencerPlaylists.Stop.Small", new IMAGE_BRUSH_SVG("Starship/MainToolbar/stop", Icon12x12));
+			Style->Set("SequencerPlaylists.Reset", new IMAGE_BRUSH_SVG("Starship/Common/Reset", Icon20x20));
+			Style->Set("SequencerPlaylists.Reset.Small", new IMAGE_BRUSH_SVG("Starship/Common/Reset", Icon12x12));
 		}
 	}
 
diff --git a/Engine/Plugins/VirtualProduction/SequencerPlaylists/Source/SequencerPlaylists/Private/SequencerPlaylistsWidgets.cpp b/Engine/Plugins/VirtualProduction/SequencerPlaylists/Source/SequencerPlaylists/Private/SequencerPlaylistsWidgets.cpp
--- a/Engine/Plugins/VirtualProduction/SequencerPlaylists/Source/SequencerPlaylists/Private/SequencerPlaylistsWidgets.cpp
+++ b/Engine/Plugins/VirtualProduction/SequencerPlaylists/Source/SequencerPlaylists/Private/SequencerPlaylistsWidgets.cpp
@@ -32,7 +32,9 @@
 #include "Widgets/Input/SComboButton.h"
 #include "Widgets/Input/SMenuAnchor.h"
 #include "Widgets/Input/SMultiLineEditableTextBox.h"
+#include "Widgets/Input/SNumericEntryBox.h"
 #include "Widgets/Input/SSearchBox.h"
+#include "Widgets/Layout/SBox.h"
 #include "Widgets/Layout/SScrollBox.h"
 #include "Widgets/Layout/SSpacer.h"
 
@@ -72,7 +74,7 @@ void SSequencerPlaylistPanel::Construct(const FArguments& InArgs, USequencerPlay
 			SNew(SVerticalBox)
 			+SVerticalBox::Slot()
 			.AutoHeight()
-			.Padding(0.0f, 0.0f, 8.0f, 0.0f)
+			.Padding(0.0f, 0.0f, 4.0f, 0.0f)
 			[
 				SNew(SHorizontalBox)
 				+ SHorizontalBox::Slot()
@@ -92,32 +94,43 @@ void SSequencerPlaylistPanel::Construct(const FArguments& InArgs, USequencerPlay
 			]
 			+ SVerticalBox::Slot()
 			.AutoHeight()
-			.Padding(8.0f, 3.0f)
+			.Padding(8.0f, 4.0f, 8.0f, 8.0f)
 			[
-				SNew(STextBlock)
-				.Font(FSequencerPlaylistsStyle::Get().GetFontStyle("SequencerPlaylists.TitleFont"))
-				.Text(TAttribute<FText>::CreateLambda([This = SharedThis(this)]()
-					{
-						if (USequencerPlaylist* Playlist = This->GetCheckedPlaylist())
+				SNew(SHorizontalBox)
+				+ SHorizontalBox::Slot()
+				.FillWidth(1.0f)
+				.VAlign(VAlign_Center)
+				[
+					SNew(STextBlock)
+					.Font(FSequencerPlaylistsStyle::Get().GetFontStyle("SequencerPlaylists.TitleFont"))
+					.Text(TAttribute<FText>::CreateLambda([this]()
 						{
-							return FText::FromString(Playlist->GetName());
-						}
+							if (USequencerPlaylist* Playlist = GetCheckedPlaylist())
+							{
+								return FText::FromString(Playlist->GetName());
+							}
 
-						return FText::GetEmpty();
-					}))
+							return FText::GetEmpty();
+						}))
+				]
+				+ SHorizontalBox::Slot()
+				.AutoWidth()
+				.VAlign(VAlign_Center)
+				[
+					Construct_Transport()
+				]
 			]
 			+ SVerticalBox::Slot()
 			.AutoHeight()
 			.MaxHeight(150.0f)
-			.Padding(8.0f, 3.0f)
+			.Padding(8.0f, 0.0f)
 			[
-				SNew(SHorizontalBox)
-				+ SHorizontalBox::Slot()
-				.FillWidth(1.0f)
+				SNew(SScrollBox)
+				+ SScrollBox::Slot()
+				.Padding(0.0f, 0.0f, 8.0f, 0.0f)
 				[
-					SNew(SScrollBox)
-					+ SScrollBox::Slot()
-					.Padding(0.0f, 0.0f, 8.0f, 0.0f)
+					SNew(SBox)
+					.MinDesiredHeight(35.0f)
 					[
 						SNew(SMultiLineEditableTextBox)
 						.Padding(0)
@@ -126,34 +139,26 @@ void SSequencerPlaylistPanel::Construct(const FArguments& InArgs, USequencerPlay
 						.Style(FSequencerPlaylistsStyle::Get(), "SequencerPlaylists.EditableTextBox")
 						.Font(FSequencerPlaylistsStyle::Get().GetFontStyle("SequencerPlaylists.DescriptionFont"))
 						.HintText(LOCTEXT("PlaylistDescriptionHint", "<playlist description>"))
-						.Text(TAttribute<FText>::CreateLambda([This = SharedThis(this)]()
+						.Text(TAttribute<FText>::CreateLambda([this]() {
+							if (USequencerPlaylist* Playlist = GetCheckedPlaylist())
 							{
-								if (USequencerPlaylist* Playlist = This->GetCheckedPlaylist())
-								{
-									return Playlist->Description;
-								}
-
-								return FText::GetEmpty();
-							}))
-						.OnTextCommitted_Lambda([This = SharedThis(this)](const FText& NewText, ETextCommit::Type CommitType)
+								return Playlist->Description;
+							}
+
+							return FText::GetEmpty();
+						}))
+						.OnTextCommitted_Lambda([this](const FText& NewText, ETextCommit::Type CommitType) {
+							if (USequencerPlaylist* Playlist = GetCheckedPlaylist())
 							{
-								if (USequencerPlaylist* Playlist = This->GetCheckedPlaylist())
-								{
-									Playlist->Description = NewText;
-								}
-							})
+								Playlist->Description = NewText;
+							}
+						})
 					]
 				]
-				+ SHorizontalBox::Slot()
-				.AutoWidth()
-				.VAlign(VAlign_Bottom)
-				[
-					Construct_Transport()
-				]
 			]
 			+ SVerticalBox::Slot()
 			.AutoHeight()
-			.Padding(8.0f, 6.0f, 8.0f, 3.0f)
+			.Padding(8.0f, 6.0f)
 			[
 				Construct_AddSearchRow()
 			]
@@ -175,28 +180,46 @@ TSharedRef<SWidget> SSequencerPlaylistPanel::Construct_LeftToolbar()
 {
 	FSequencerPlaylistsModule& Module = static_cast<FSequencerPlaylistsModule&>(ISequencerPlaylistsModule::Get());
 	FSlimHorizontalToolBarBuilder ToolBarBuilder(Module.GetCommandList(), FMultiBoxCustomization::None, nullptr, true);
+	ToolBarBuilder.SetStyle(&FSequencerPlaylistsStyle::Get(), "SequencerPlaylists.MainToolbar");
 
 	ToolBarBuilder.BeginSection("Playlists");
 	{
 		ToolBarBuilder.AddToolBarButton(
 			FUIAction(FExecuteAction::CreateSP(this, &SSequencerPlaylistPanel::OnNewPlaylist)),
 			NAME_None,
-			LOCTEXT("NewPlaylistOptions", "New Playlist"),
-			LOCTEXT("NewPlaylistOptionsToolTip", "New Playlist"),
-			FSlateIcon(FSequencerPlaylistsStyle::Get().GetStyleSetName(), "SequencerPlaylists.NewPlaylist.Background", NAME_None, "SequencerPlaylists.NewPlaylist.Overlay"));
+			LOCTEXT("NewPlaylist", "New Playlist"),
+			LOCTEXT("NewPlaylistToolTip", "New Playlist"),
+			FSlateIcon(FSequencerPlaylistsStyle::Get().GetStyleSetName(), "SequencerPlaylists.NewPlaylist"));
 
 		ToolBarBuilder.AddToolBarButton(
-			FUIAction(FExecuteAction::CreateSP(this, &SSequencerPlaylistPanel::OnSavePlaylistAs)),
+			FUIAction(FExecuteAction::CreateSP(this, &SSequencerPlaylistPanel::OnSavePlaylist)),
 			NAME_None,
-			LOCTEXT("SavePlaylist", "Save Playlist As..."),
-			LOCTEXT("SavePlaylistTooltip", "Save Playlist As..."),
-			FSlateIcon(FSequencerPlaylistsStyle::Get().GetStyleSetName(), "SequencerPlaylists.SavePlaylistAs"));
+			LOCTEXT("SavePlaylist", "Save Playlist"),
+			LOCTEXT("SavePlaylistTooltip", "Save Playlist"),
+			FSlateIcon(FSequencerPlaylistsStyle::Get().GetStyleSetName(), "SequencerPlaylists.SavePlaylist"));
+
+		ToolBarBuilder.AddComboButton(
+			FUIAction(),
+			FOnGetContent::CreateLambda([this]() {
+				FMenuBuilder MenuBuilder(true, nullptr);
+				MenuBuilder.AddMenuEntry(
+					LOCTEXT("SavePlaylistAs", "Save Playlist As..."),
+					LOCTEXT("SavePlaylistAsTooltip", "Save Playlist As..."),
+					FSlateIcon(),
+					FUIAction(FExecuteAction::CreateSP(this, &SSequencerPlaylistPanel::OnSavePlaylistAs))
+				);
+				return MenuBuilder.MakeWidget();
+			}),
+			LOCTEXT("SavePlaylistOptions", "Save Playlist Options"),
+			LOCTEXT("SavePlaylistOptionsTooltip", "Save Playlist Options"),
+			TAttribute<FSlateIcon>(),
+			true);
 
 		ToolBarBuilder.AddComboButton(
 			FUIAction(),
-			FOnGetContent::CreateSP(SharedThis(this), &SSequencerPlaylistPanel::BuildOpenPlaylistMenu),
-			LOCTEXT("NewPlaylistOptions", "Open Playlist"),
-			LOCTEXT("NewPlaylistOptionsToolTip", "Open Playlist"),
+			FOnGetContent::CreateSP(this, &SSequencerPlaylistPanel::BuildOpenPlaylistMenu),
+			LOCTEXT("OpenPlaylist", "Open Playlist"),
+			LOCTEXT("OpenPlaylistToolTip", "Open Playlist"),
 			FSlateIcon(FSequencerPlaylistsStyle::Get().GetStyleSetName(), "SequencerPlaylists.OpenPlaylist"));
 	}
 	ToolBarBuilder.EndSection();
@@ -209,26 +232,27 @@ END_SLATE_FUNCTION_BUILD_OPTIMIZATION
 BEGIN_SLATE_FUNCTION_BUILD_OPTIMIZATION
 TSharedRef<SWidget> SSequencerPlaylistPanel::Construct_RightToolbar()
 {
-	TSharedRef<SSequencerPlaylistPanel> This = SharedThis(this);
-
 	FSequencerPlaylistsModule& Module = static_cast<FSequencerPlaylistsModule&>(ISequencerPlaylistsModule::Get());
 	FSlimHorizontalToolBarBuilder ToolBarBuilder(Module.GetCommandList(), FMultiBoxCustomization::None, nullptr, true);
+	ToolBarBuilder.SetStyle(&FSequencerPlaylistsStyle::Get(), "SequencerPlaylists.MainToolbar");
 
-	ToolBarBuilder.BeginSection("TriggerMode");
+	ToolBarBuilder.BeginSection("PlayMode");
 	{
-		TSharedRef<SWidget> TriggerToggle = SNew(SCheckBox)
-			.Padding(FMargin(10.0f, 4.0f))
-			.ToolTipText(LOCTEXT("ToggleTriggerToolTip", "Enables and disables Trigger Mode"))
-			.Style(FEditorStyle::Get(), "ToggleButtonCheckbox")
-			.IsChecked_Lambda([This]() { return This->InTriggerMode() ? ECheckBoxState::Checked : ECheckBoxState::Unchecked; })
-			.OnCheckStateChanged_Lambda([This](ECheckBoxState InState) { This->bTriggerMode = (InState == ECheckBoxState::Checked); })
-			[
-				SNew(STextBlock)
-				.Font(FEditorStyle::Get().GetFontStyle("FontAwesome.16"))
-				.Text(FText::FromString(FString(TEXT("\xf0e7"))))
-			];
-
-		ToolBarBuilder.AddWidget(TriggerToggle);
+		ToolBarBuilder.AddToolBarButton(
+			FUIAction(
+				FExecuteAction::CreateLambda([this]() {
+					bPlayMode = !bPlayMode;
+				}),
+				FCanExecuteAction(),
+				FGetActionCheckState::CreateLambda([this]() {
+					return InPlayMode() ? ECheckBoxState::Checked : ECheckBoxState::Unchecked;
+				})
+			),
+			NAME_None,
+			LOCTEXT("TogglePlayMode", "Toggle Play Mode"),
+			LOCTEXT("TogglePlayModeToolTip", "Enables or disables Play Mode"),
+			FSlateIcon(FSequencerPlaylistsStyle::Get().GetStyleSetName(), "SequencerPlaylists.PlayMode"),
+			EUserInterfaceActionType::ToggleButton);
 	}
 	ToolBarBuilder.EndSection();
 
@@ -246,7 +270,7 @@ TSharedRef<SWidget> SSequencerPlaylistPanel::Construct_Transport()
 		[
 			SNew(SButton)
 			.ButtonStyle(FSequencerPlaylistsStyle::Get(), "SequencerPlaylists.TransportButton.Play")
-			.ContentPadding(0)
+			.ContentPadding(FMargin(0.0f, 2.0f))
 			.HAlign(HAlign_Center)
 			.VAlign(VAlign_Center)
 			.OnClicked(this, &SSequencerPlaylistPanel::HandleClicked_PlayAll)
@@ -262,7 +286,7 @@ TSharedRef<SWidget> SSequencerPlaylistPanel::Construct_Transport()
 		[
 			SNew(SButton)
 			.ButtonStyle(FSequencerPlaylistsStyle::Get(), "SequencerPlaylists.TransportButton.Stop")
-			.ContentPadding(0)
+			.ContentPadding(FMargin(0.0f, 2.0f))
 			.HAlign(HAlign_Center)
 			.VAlign(VAlign_Center)
 			.OnClicked(this, &SSequencerPlaylistPanel::HandleClicked_StopAll)
@@ -278,15 +302,15 @@ TSharedRef<SWidget> SSequencerPlaylistPanel::Construct_Transport()
 		[
 			SNew(SButton)
 			.ButtonStyle(FSequencerPlaylistsStyle::Get(), "SequencerPlaylists.TransportButton.Reset")
-			.ContentPadding(FMargin(0.0f, 1.0f, 0.0f, 0.0f))
+			.ContentPadding(FMargin(0.0f, 2.0f))
 			.HAlign(HAlign_Center)
 			.VAlign(VAlign_Center)
 			.OnClicked(this, &SSequencerPlaylistPanel::HandleClicked_ResetAll)
 			.ToolTipText(LOCTEXT("ResetAllButtonToolTip", "Reset all items"))
 			[
-				SNew(STextBlock)
-				.Font(FEditorStyle::Get().GetFontStyle("FontAwesome.14"))
-				.Text(FEditorFontGlyphs::Undo)
+				SNew(SImage)
+				.Image(FSequencerPlaylistsStyle::Get().GetBrush("SequencerPlaylists.Reset"))
+				.ColorAndOpacity(FSlateColor::UseForeground())
 			]
 		];
 }
@@ -324,13 +348,13 @@ TSharedRef<SWidget> SSequencerPlaylistPanel::Construct_ItemListView()
 	return SAssignNew(ItemListView, SListView<TSharedPtr<FSequencerPlaylistRowData>>)
 		.SelectionMode(ESelectionMode::Single) // See TODO in HandleAcceptDrop
 		.ListItemsSource(&ItemRows)
-		.OnGenerateRow_Lambda([This = SharedThis(this)](TSharedPtr<FSequencerPlaylistRowData> InData, const TSharedRef<STableViewBase>& OwnerTableView)
+		.OnGenerateRow_Lambda([this](TSharedPtr<FSequencerPlaylistRowData> InData, const TSharedRef<STableViewBase>& OwnerTableView)
 			{
 				return SNew(SSequencerPlaylistItemWidget, InData, OwnerTableView)
-					.TriggerMode(TAttribute<bool>::Create(TAttribute<bool>::FGetter::CreateSP(This, &SSequencerPlaylistPanel::InTriggerMode)))
-					.IsPlaying_Lambda([This, WeakItem = InData->WeakItem]() {
+					.PlayMode(this, &SSequencerPlaylistPanel::InPlayMode)
+					.IsPlaying_Lambda([this, WeakItem = InData->WeakItem]() {
 						USequencerPlaylistItem* Item = WeakItem.Get();
-						USequencerPlaylistPlayer* Player = This->WeakPlayer.Get();
+						USequencerPlaylistPlayer* Player = WeakPlayer.Get();
 						if (Item && Player)
 						{
 							return Player->IsPlaying(Item);
@@ -340,19 +364,19 @@ TSharedRef<SWidget> SSequencerPlaylistPanel::Construct_ItemListView()
 							return false;
 						}
 					})
-					.OnPlayClicked(This, &SSequencerPlaylistPanel::HandleClicked_Item_Play)
-					.OnStopClicked(This, &SSequencerPlaylistPanel::HandleClicked_Item_Stop)
-					.OnResetClicked(This, &SSequencerPlaylistPanel::HandleClicked_Item_Reset)
-					.OnRemoveClicked(This, &SSequencerPlaylistPanel::HandleClicked_Item_Remove)
-					.OnIsPropertyVisible(This, &SSequencerPlaylistPanel::HandleItemDetailsIsPropertyVisible)
-					.OnCanAcceptDrop(This, &SSequencerPlaylistPanel::HandleCanAcceptDrop)
-					.OnAcceptDrop(This, &SSequencerPlaylistPanel::HandleAcceptDrop);
+					.OnPlayClicked(this, &SSequencerPlaylistPanel::HandleClicked_Item_Play)
+					.OnStopClicked(this, &SSequencerPlaylistPanel::HandleClicked_Item_Stop)
+					.OnResetClicked(this, &SSequencerPlaylistPanel::HandleClicked_Item_Reset)
+					.OnRemoveClicked(this, &SSequencerPlaylistPanel::HandleClicked_Item_Remove)
+					.OnIsPropertyVisible(this, &SSequencerPlaylistPanel::HandleItemDetailsIsPropertyVisible)
+					.OnCanAcceptDrop(this, &SSequencerPlaylistPanel::HandleCanAcceptDrop)
+					.OnAcceptDrop(this, &SSequencerPlaylistPanel::HandleAcceptDrop);
 			})
 		.HeaderRow(
 			SNew(SHeaderRow)
 			+ SHeaderRow::Column(ColumnName_HoverTransport)
 				.DefaultLabel(FText::GetEmpty())
-				.FillSized(20.0f)
+				.FillSized(30.0f)
 			+ SHeaderRow::Column(ColumnName_Items)
 				.DefaultLabel(LOCTEXT("ColumnLabelItems", "Playlist Items"))
 				.FillWidth(1.0f)
@@ -365,10 +389,10 @@ TSharedRef<SWidget> SSequencerPlaylistPanel::Construct_ItemListView()
 				.HAlignCell(HAlign_Center)
 			+ SHeaderRow::Column(ColumnName_Loop)
 				.DefaultLabel(LOCTEXT("ColumnLabelLoop", "Loop"))
-				.FillSized(40.0f)
+				.FillSized(80.0f)
 			+ SHeaderRow::Column(ColumnName_HoverDetails)
 				.DefaultLabel(FText::GetEmpty())
-				.FillSized(20.0f)
+				.FillSized(30.0f)
 				.HAlignCell(HAlign_Center)
 				.VAlignCell(VAlign_Center)
 		);
@@ -536,26 +560,45 @@ bool GetSavePlaylistPackageName(FString& OutName)
 }
 
 
-void SSequencerPlaylistPanel::OnSavePlaylistAs()
+void SSequencerPlaylistPanel::OnSavePlaylist()
 {
-	USequencerPlaylist* Playlist = GetCheckedPlaylist();
-	if (!Playlist)
+	if (!WeakLoadedPlaylist.IsValid() ||
+		!WeakLoadedPlaylist->GetPackage() ||
+		!WeakLoadedPlaylist->GetPackage()->GetLoadedPath().HasPackageName())
 	{
+		OnSavePlaylistAs();
 		return;
 	}
 
+	FString PackageName = WeakLoadedPlaylist->GetPackage()->GetLoadedPath().GetPackageName();
+	SavePlaylist(PackageName);
+}
+
+
+void SSequencerPlaylistPanel::OnSavePlaylistAs()
+{
 	FString PackageName;
 	if (!GetSavePlaylistPackageName(PackageName))
 	{
 		return;
 	}
 
+	SavePlaylist(PackageName);
+}
+
+void SSequencerPlaylistPanel::SavePlaylist(const FString& PackageName)
+{
+	USequencerPlaylist* Playlist = GetCheckedPlaylist();
+	if (!Playlist)
+	{
+		return;
+	}
+
 	FScopedTransaction Transaction(LOCTEXT("SavePlaylistTransaction", "Save Sequencer Playlist"));
 
-	// Saving into a new package
-	const FString NewAssetName = FPackageName::GetLongPackageAssetName(PackageName);
-	UPackage* NewPackage = CreatePackage(*PackageName);
-	USequencerPlaylist* NewPlaylist = NewObject<USequencerPlaylist>(NewPackage, *NewAssetName, RF_Public | RF_Standalone | RF_Transactional);
+	const FString AssetName = FPackageName::GetLongPackageAssetName(PackageName);
+	UPackage* Package = CreatePackage(*PackageName);
+	USequencerPlaylist* NewPlaylist = NewObject<USequencerPlaylist>(Package, *AssetName, RF_Public | RF_Standalone | RF_Transactional);
 
 	if (NewPlaylist)
 	{
@@ -572,8 +615,10 @@ void SSequencerPlaylistPanel::OnSavePlaylistAs()
 		NewPlaylist->MarkPackageDirty();
 		FAssetRegistryModule::AssetCreated(NewPlaylist);
 
-		FEditorFileUtils::PromptForCheckoutAndSave({ NewPackage }, false, false);
+		FEditorFileUtils::PromptForCheckoutAndSave({ Package }, false, false);
 	}
+
+	WeakLoadedPlaylist = NewPlaylist;
 }
 
 
@@ -604,6 +649,8 @@ void SSequencerPlaylistPanel::OnLoadPlaylist(const FAssetData& InPreset)
 		}
 
 		RegenerateRows();
+
+		WeakLoadedPlaylist = ImportedPlaylist;
 	}
 }
 
@@ -616,10 +663,12 @@ void SSequencerPlaylistPanel::OnNewPlaylist()
 		return;
 	}
 
-	Playlist->Rename(*MakeUniqueObjectName(GetTransientPackage(), USequencerPlaylist::StaticClass(), "UntitledPlaylist").ToString());
+	Playlist->Rename(*MakeUniqueObjectName(GetTransientPackage(), USequencerPlaylist::StaticClass(), "SequencerPlaylist").ToString());
 	Playlist->Description = FText::GetEmpty();
 	Playlist->Items.Empty();
 	RegenerateRows();
+
+	WeakLoadedPlaylist = nullptr;
 }
 
 
@@ -687,13 +736,13 @@ FReply SSequencerPlaylistPanel::HandleClicked_AddSequence()
 	}
 
 	Playlist->Items.Add(NewObject<USequencerPlaylistItem_Sequence>(Playlist));
-	RegenerateRows(); // TODO: Incremental view update? (UETOOL-4644)
+	RegenerateRows();
 
 	return FReply::Handled();
 }
 
 
-FReply SSequencerPlaylistPanel::HandleClicked_Item_Play(TSharedPtr<SSequencerPlaylistItemWidget> ItemWidget)
+FReply SSequencerPlaylistPanel::HandleClicked_Item_Play(SSequencerPlaylistItemWidget& ItemWidget)
 {
 	USequencerPlaylistPlayer* Player = WeakPlayer.Get();
 	if (!ensure(Player))
@@ -701,12 +750,12 @@ FReply SSequencerPlaylistPanel::HandleClicked_Item_Play(TSharedPtr<SSequencerPla
 		return FReply::Unhandled();
 	}
 
-	Player->PlayItem(ItemWidget->GetItem());
+	Player->PlayItem(ItemWidget.GetItem());
 	return FReply::Handled();
 }
 
 
-FReply SSequencerPlaylistPanel::HandleClicked_Item_Stop(TSharedPtr<SSequencerPlaylistItemWidget> ItemWidget)
+FReply SSequencerPlaylistPanel::HandleClicked_Item_Stop(SSequencerPlaylistItemWidget& ItemWidget)
 {
 	USequencerPlaylistPlayer* Player = WeakPlayer.Get();
 	if (!ensure(Player))
@@ -714,12 +763,12 @@ FReply SSequencerPlaylistPanel::HandleClicked_Item_Stop(TSharedPtr<SSequencerPla
 		return FReply::Unhandled();
 	}
 
-	Player->StopItem(ItemWidget->GetItem());
+	Player->StopItem(ItemWidget.GetItem());
 	return FReply::Handled();
 }
 
 
-FReply SSequencerPlaylistPanel::HandleClicked_Item_Reset(TSharedPtr<SSequencerPlaylistItemWidget> ItemWidget)
+FReply SSequencerPlaylistPanel::HandleClicked_Item_Reset(SSequencerPlaylistItemWidget& ItemWidget)
 {
 	USequencerPlaylistPlayer* Player = WeakPlayer.Get();
 	if (!ensure(Player))
@@ -727,12 +776,12 @@ FReply SSequencerPlaylistPanel::HandleClicked_Item_Reset(TSharedPtr<SSequencerPl
 		return FReply::Unhandled();
 	}
 
-	Player->ResetItem(ItemWidget->GetItem());
+	Player->ResetItem(ItemWidget.GetItem());
 	return FReply::Handled();
 }
 
 
-FReply SSequencerPlaylistPanel::HandleClicked_Item_Remove(TSharedPtr<SSequencerPlaylistItemWidget> ItemWidget)
+FReply SSequencerPlaylistPanel::HandleClicked_Item_Remove(SSequencerPlaylistItemWidget& ItemWidget)
 {
 	USequencerPlaylist* Playlist = GetCheckedPlaylist();
 	if (!Playlist)
@@ -740,8 +789,8 @@ FReply SSequencerPlaylistPanel::HandleClicked_Item_Remove(TSharedPtr<SSequencerP
 		return FReply::Unhandled();
 	}
 
-	ensure(Playlist->Items.RemoveSingle(ItemWidget->GetItem()));
-	RegenerateRows(); // TODO: Incremental view update? (UETOOL-4644)
+	ensure(Playlist->Items.RemoveSingle(ItemWidget.GetItem()));
+	RegenerateRows();
 	return FReply::Handled();
 }
 
@@ -828,7 +877,7 @@ FReply SSequencerPlaylistPanel::HandleAcceptDrop(const FDragDropEvent& DragDropE
 				USequencerPlaylistItem* ItemToMove = Items[SrcIndex];
 				Items.RemoveAt(SrcIndex);
 				Items.Insert(ItemToMove, DestIndex);
-				RegenerateRows(); // TODO: Incremental view update? (UETOOL-4644)
+				RegenerateRows();
 				return FReply::Handled();
 			}
 		}
@@ -880,7 +929,9 @@ void SSequencerPlaylistItemWidget::Construct(const FArguments& InArgs, TSharedPt
 	check(InRowData && InRowData->WeakItem.IsValid());
 	RowData = InRowData;
 
-	TriggerMode = InArgs._TriggerMode;
+	LoopMode = InRowData->WeakItem->NumLoops == 0 ? ELoopMode::None : ELoopMode::Finite;
+
+	PlayMode = InArgs._PlayMode;
 	IsPlaying = InArgs._IsPlaying;
 
 	PlayClickedDelegate = InArgs._OnPlayClicked;
@@ -895,6 +946,7 @@ void SSequencerPlaylistItemWidget::Construct(const FArguments& InArgs, TSharedPt
 			.OnCanAcceptDrop(InArgs._OnCanAcceptDrop)
 			.OnAcceptDrop(InArgs._OnAcceptDrop)
 			.OnDragDetected(this, &SSequencerPlaylistItemWidget::HandleDragDetected)
+			.Style(FSequencerPlaylistsStyle::Get(), "SequencerPlaylists.ItemRow")
 		, OwnerTableView);
 }
 
@@ -902,10 +954,7 @@ void SSequencerPlaylistItemWidget::Construct(const FArguments& InArgs, TSharedPt
 void SSequencerPlaylistItemWidget::ConstructChildren(ETableViewMode::Type InOwnerTableMode, const TAttribute<FMargin>& InPadding, const TSharedRef<SWidget>& InContent)
 {
 	// We wrap InContent with an overlay which facilitates dimming the row, obstructing hit test of
-	// the controls behind, and centering Trigger Mode transport controls over the row on hover.
-
-	TSharedRef<SSequencerPlaylistItemWidget> This = SharedThis(this);
-
+	// the controls behind, and centering Play Mode transport controls over the row on hover.
 	Content = InContent;
 
 	ChildSlot
@@ -913,10 +962,24 @@ void SSequencerPlaylistItemWidget::ConstructChildren(ETableViewMode::Type InOwne
 	[
 		SNew(SOverlay)
 		+ SOverlay::Slot()
+		.Padding(1.0f)
+		[
+			SNew(SImage)
+			.Image(FSequencerPlaylistsStyle::Get().GetBrush("SequencerPlaylists.ItemRow.BackgroundOuter"))
+		]
+		+ SOverlay::Slot()
+		.Padding(8.0f, 4.0f)
+		[
+			SNew(SImage)
+			.Image(FSequencerPlaylistsStyle::Get().GetBrush("SequencerPlaylists.ItemRow.BackgroundInner"))
+		]
+		+ SOverlay::Slot()
+		.Padding(0.0f, 4.0f)
 		[
 			InContent
 		]
 		+ SOverlay::Slot()
+		.Padding(1.0f)
 		[
 			SNew(SImage)
 			.Visibility(this, &SSequencerPlaylistItemWidget::GetRowDimmingVisibility)
@@ -927,16 +990,16 @@ void SSequencerPlaylistItemWidget::ConstructChildren(ETableViewMode::Type InOwne
 		.VAlign(VAlign_Center)
 		[
 			SNew(SHorizontalBox)
-			.Visibility(this, &SSequencerPlaylistItemWidget::GetTriggerModeTransportVisibility)
+			.Visibility(this, &SSequencerPlaylistItemWidget::GetPlayModeTransportVisibility)
 			+ SHorizontalBox::Slot()
 			.AutoWidth()
 			[
 				SNew(SButton)
 				.ButtonStyle(FSequencerPlaylistsStyle::Get(), "SequencerPlaylists.TransportButton.Play")
-				.ContentPadding(0)
+				.ContentPadding(6.0f)
 				.HAlign(HAlign_Center)
 				.VAlign(VAlign_Center)
-				.OnClicked_Lambda([This]() { return This->PlayClickedDelegate.Execute(This); })
+				.OnClicked_Lambda([this]() { return PlayClickedDelegate.Execute(*this); })
 				.ToolTipText(PlayItemTooltipText)
 				[
 					SNew(SImage)
@@ -949,10 +1012,10 @@ void SSequencerPlaylistItemWidget::ConstructChildren(ETableViewMode::Type InOwne
 			[
 				SNew(SButton)
 				.ButtonStyle(FSequencerPlaylistsStyle::Get(), "SequencerPlaylists.TransportButton.Stop")
-				.ContentPadding(0)
+				.ContentPadding(6.0f)
 				.HAlign(HAlign_Center)
 				.VAlign(VAlign_Center)
-				.OnClicked_Lambda([This]() { return This->StopClickedDelegate.Execute(This); })
+				.OnClicked_Lambda([this]() { return StopClickedDelegate.Execute(*this); })
 				.ToolTipText(StopItemTooltipText)
 				[
 					SNew(SImage)
@@ -965,15 +1028,15 @@ void SSequencerPlaylistItemWidget::ConstructChildren(ETableViewMode::Type InOwne
 			[
 				SNew(SButton)
 				.ButtonStyle(FSequencerPlaylistsStyle::Get(), "SequencerPlaylists.TransportButton.Reset")
-				.ContentPadding(FMargin(0.0f, 1.0f, 0.0f, 0.0f))
+				.ContentPadding(6.0f)
 				.HAlign(HAlign_Center)
 				.VAlign(VAlign_Center)
-				.OnClicked_Lambda([This]() { return This->ResetClickedDelegate.Execute(This); })
+				.OnClicked_Lambda([this]() { return ResetClickedDelegate.Execute(*this); })
 				.ToolTipText(ResetItemTooltipText)
 				[
-					SNew(STextBlock)
-					.Font(FEditorStyle::Get().GetFontStyle("FontAwesome.14"))
-					.Text(FEditorFontGlyphs::Undo)
+					SNew(SImage)
+					.Image(FSequencerPlaylistsStyle::Get().GetBrush("SequencerPlaylists.Reset"))
+					.ColorAndOpacity(FSlateColor::UseForeground())
 				]
 			]
 		]
@@ -993,7 +1056,7 @@ TSharedRef<SWidget> SSequencerPlaylistItemWidget::GenerateWidgetForColumn(const
 		return SNullWidget::NullWidget;
 	}
 
-	TSharedRef<SSequencerPlaylistItemWidget> This = SharedThis(this);
+	TWeakObjectPtr<USequencerPlaylistItem> WeakItem = Item;
 
 	static const FName PropertyEditorModuleName("PropertyEditor");
 	FPropertyEditorModule& PropertyEditorModule = FModuleManager::GetModuleChecked<FPropertyEditorModule>(PropertyEditorModuleName);
@@ -1004,14 +1067,15 @@ TSharedRef<SWidget> SSequencerPlaylistItemWidget::GenerateWidgetForColumn(const
 	{
 		return SNew(SVerticalBox)
 			+ SVerticalBox::Slot()
+			.Padding(8.0f, 4.0f, 0.0f, 0.0f)
 			[
 				SNew(SButton)
 				.ButtonStyle(FSequencerPlaylistsStyle::Get(), "SequencerPlaylists.HoverTransport.Play")
-				.ContentPadding(FMargin(1.0f, 3.0f, 0.0f, 0.0f))
+				.ContentPadding(0)
 				.HAlign(HAlign_Center)
 				.VAlign(VAlign_Center)
-				.Visibility_Lambda([This]() { return (This->IsHovered() && !This->InTriggerMode()) ? EVisibility::Visible : EVisibility::Hidden; })
-				.OnClicked_Lambda([This]() { return This->PlayClickedDelegate.Execute(This); })
+				.Visibility_Lambda([this]() { return (IsHovered() && !InPlayMode()) ? EVisibility::Visible : EVisibility::Hidden; })
+				.OnClicked_Lambda([this]() { return PlayClickedDelegate.Execute(*this); })
 				.ToolTipText(PlayItemTooltipText)
 				[
 					SNew(SImage)
@@ -1020,15 +1084,16 @@ TSharedRef<SWidget> SSequencerPlaylistItemWidget::GenerateWidgetForColumn(const
 				]
 			]
 			+ SVerticalBox::Slot()
+			.Padding(8.0f, 0.0f, 0.0f, 0.0f)
 			[
 				SNew(SButton)
 				.ButtonStyle(FSequencerPlaylistsStyle::Get(), "SequencerPlaylists.HoverTransport.Stop")
-				.ContentPadding(FMargin(0.0f, 0.0f, 0.0f, 1.0f))
+				.ContentPadding(0)
 				.HAlign(HAlign_Center)
 				.VAlign(VAlign_Center)
-				.Visibility_Lambda([This]() { return (This->IsPlaying.Get() || (This->IsHovered() && !This->InTriggerMode())) ? EVisibility::Visible : EVisibility::Hidden; })
-				.ForegroundColor_Lambda([This]() { return This->IsPlaying.Get() ? FStyleColors::AccentRed : FSlateColor::UseStyle(); })
-				.OnClicked_Lambda([This]() { return This->StopClickedDelegate.Execute(This); })
+				.Visibility_Lambda([this]() { return (IsPlaying.Get() || (IsHovered() && !InPlayMode())) ? EVisibility::Visible : EVisibility::Hidden; })
+				.ForegroundColor_Lambda([this]() { return IsPlaying.Get() ? FStyleColors::AccentRed : FSlateColor::UseStyle(); })
+				.OnClicked_Lambda([this]() { return StopClickedDelegate.Execute(*this); })
 				.ToolTipText(StopItemTooltipText)
 				[
 					SNew(SImage)
@@ -1037,19 +1102,20 @@ TSharedRef<SWidget> SSequencerPlaylistItemWidget::GenerateWidgetForColumn(const
 				]
 			]
 			+ SVerticalBox::Slot()
+			.Padding(8.0f, 0.0f, 0.0f, 4.0f)
 			[
 				SNew(SButton)
 				.ButtonStyle(FSequencerPlaylistsStyle::Get(), "SequencerPlaylists.HoverTransport.Reset")
-				.ContentPadding(FMargin(1.0f, 0.0f, 0.0f, 3.0f))
+				.ContentPadding(0)
 				.HAlign(HAlign_Center)
 				.VAlign(VAlign_Center)
-				.Visibility_Lambda([This]() { return (This->IsHovered() && !This->InTriggerMode()) ? EVisibility::Visible : EVisibility::Hidden; })
-				.OnClicked_Lambda([This]() { return This->ResetClickedDelegate.Execute(This); })
+				.Visibility_Lambda([this]() { return (IsHovered() && !InPlayMode()) ? EVisibility::Visible : EVisibility::Hidden; })
+				.OnClicked_Lambda([this]() { return ResetClickedDelegate.Execute(*this); })
 				.ToolTipText(ResetItemTooltipText)
 				[
-					SNew(STextBlock)
-					.Font(FEditorStyle::Get().GetFontStyle("FontAwesome.9"))
-					.Text(FEditorFontGlyphs::Undo)
+					SNew(SImage)
+					.Image(FSequencerPlaylistsStyle::Get().GetBrush("SequencerPlaylists.Reset.Small"))
+					.ColorAndOpacity(FSlateColor::UseForeground())
 				]
 			];
 	}
@@ -1064,7 +1130,7 @@ TSharedRef<SWidget> SSequencerPlaylistItemWidget::GenerateWidgetForColumn(const
 			SinglePropParams);
 
 		return SNew(SBox)
-			.Padding(FMargin(8.0f, 0.0f, 0.0f, 0.0f))
+			.Padding(FMargin(0.0f, 0.0f, 0.0f, 4.0f))
 			[
 				PropView.ToSharedRef()
 			];
@@ -1091,39 +1157,137 @@ TSharedRef<SWidget> SSequencerPlaylistItemWidget::GenerateWidgetForColumn(const
 
 	if (ColumnName == SSequencerPlaylistPanel::ColumnName_Hold)
 	{
-		TSharedPtr<ISinglePropertyView> PropView = PropertyEditorModule.CreateSingleProperty(Item,
-			GET_MEMBER_NAME_CHECKED(USequencerPlaylistItem, bHoldAtFirstFrame),
-			SinglePropParams);
+		TSharedRef<SWidget> HoldToggle = SNew(SCheckBox)
+			.Padding(FMargin(4.0f, 2.0f))
+			.HAlign(HAlign_Center)
+			.ToolTipText(LOCTEXT("ToggleHoldToolTip", "Enables or disables Hold At First Frame"))
+			.Style(FEditorStyle::Get(), "ToggleButtonCheckbox")
+			.IsChecked_Lambda([WeakItem]() { return (WeakItem.IsValid() && WeakItem->bHoldAtFirstFrame) ? ECheckBoxState::Checked : ECheckBoxState::Unchecked; })
+			.OnCheckStateChanged_Lambda([WeakItem](ECheckBoxState InState) {
+				if (WeakItem.IsValid())
+				{
+					WeakItem->bHoldAtFirstFrame = (InState == ECheckBoxState::Checked);
+				}
+			})
+			[
+				SNew(SImage)
+				.Image(FSequencerPlaylistsStyle::Get().GetBrush("SequencerPlaylists.HoldFrame"))
+				.ColorAndOpacity(FSlateColor::UseForeground())
+			];
 
-		return PropView.ToSharedRef();
+		return SNew(SVerticalBox)
+			+ SVerticalBox::Slot()
+			.VAlign(VAlign_Center)
+			[
+				HoldToggle
+			]
+			+ SVerticalBox::Slot()
+			[
+				SNullWidget::NullWidget
+			];
 	}
 
 	if (ColumnName == SSequencerPlaylistPanel::ColumnName_Loop)
 	{
-		TSharedPtr<ISinglePropertyView> PropView = PropertyEditorModule.CreateSingleProperty(Item,
-			GET_MEMBER_NAME_CHECKED(USequencerPlaylistItem, NumLoops),
-			SinglePropParams);
+		TSharedRef<SWidget> LoopModeCombo = SNew(SComboButton)
+			.ContentPadding(0)
+			.ComboButtonStyle(FAppStyle::Get(), "SimpleComboButton")
+			.OnGetMenuContent_Lambda([this]() {
+				FMenuBuilder MenuBuilder(true, nullptr);
+				MenuBuilder.AddMenuEntry(
+					LOCTEXT("LoopModeNone", "Don't Loop"),
+					LOCTEXT("LoopModeNoneTooltip", "Don't Loop"),
+					FSlateIcon(FSequencerPlaylistsStyle::Get().GetStyleSetName(), "SequencerPlaylists.Loop.Disabled"),
+					FUIAction(
+						FExecuteAction::CreateLambda([this]() { SetLoopMode(ELoopMode::None); }),
+						FCanExecuteAction(),
+						FGetActionCheckState::CreateLambda([this]() { return LoopMode == ELoopMode::None ? ECheckBoxState::Checked : ECheckBoxState::Unchecked; })
+					),
+					NAME_None,
+					EUserInterfaceActionType::Check
+				);
+				MenuBuilder.AddMenuEntry(
+					LOCTEXT("LoopModeFinite", "Loop N Times"),
+					LOCTEXT("LoopModeFiniteTooltip", "Loop N Times"),
+					FSlateIcon(FSequencerPlaylistsStyle::Get().GetStyleSetName(), "SequencerPlaylists.Loop.Finite"),
+					FUIAction(
+						FExecuteAction::CreateLambda([this]() { SetLoopMode(ELoopMode::Finite); }),
+						FCanExecuteAction(),
+						FGetActionCheckState::CreateLambda([this]() { return LoopMode == ELoopMode::Finite ? ECheckBoxState::Checked : ECheckBoxState::Unchecked; })
+					),
+					NAME_None,
+					EUserInterfaceActionType::Check
+				);
+				return MenuBuilder.MakeWidget();
+			})
+			.ButtonContent()
+			[
+				SNew(SBox)
+				.VAlign(VAlign_Center)
+				.HAlign(HAlign_Center)
+				.WidthOverride(16.0f)
+				.HeightOverride(16.0f)
+				[
+					SNew(SImage)
+					.ColorAndOpacity(FSlateColor::UseForeground())
+					.Image_Lambda([this]() -> const FSlateBrush* {
+						switch (LoopMode) {
+							case ELoopMode::None: return FSequencerPlaylistsStyle::Get().GetBrush("SequencerPlaylists.Loop.Disabled");
+							case ELoopMode::Finite: return FSequencerPlaylistsStyle::Get().GetBrush("SequencerPlaylists.Loop.Finite");
+							default: checkNoEntry(); return nullptr;
+						}
+					})
+				]
+			];
 
-		return PropView.ToSharedRef();
+		return SNew(SVerticalBox)
+			+ SVerticalBox::Slot()
+			[
+				SNew(SHorizontalBox)
+				+ SHorizontalBox::Slot()
+				.AutoWidth()
+				.VAlign(VAlign_Center)
+				[
+					LoopModeCombo
+				]
+				+ SHorizontalBox::Slot()
+				[
+					SNew(SNumericEntryBox<int32>)
+					.AllowSpin(true)
+					.Value(TAttribute<TOptional<int32>>::CreateLambda([WeakItem]() { return WeakItem.IsValid() ? WeakItem->NumLoops : TOptional<int32>(); }))
+					.Visibility_Lambda([this]() { return LoopMode == ELoopMode::Finite ? EVisibility::Visible : EVisibility::Hidden; })
+					.OnValueCommitted_Lambda([WeakItem](int32 NewValue, ETextCommit::Type) {
+						if (WeakItem.IsValid())
+						{
+							WeakItem->NumLoops = NewValue;
+						}
+					})
+				]
+			]
+			+ SVerticalBox::Slot()
+			[
+				SNullWidget::NullWidget
+			];
 	}
 
 	if (ColumnName == SSequencerPlaylistPanel::ColumnName_HoverDetails)
 	{
 		return SAssignNew(DetailsAnchor, SMenuAnchor)
+			.Padding(FMargin(0.0f, 0.0f, 8.0f, 0.0f))
 			.Placement(MenuPlacement_MenuLeft)
-			.OnGetMenuContent(This, &SSequencerPlaylistItemWidget::EnsureSelectedAndBuildContextMenu)
+			.OnGetMenuContent(this, &SSequencerPlaylistItemWidget::EnsureSelectedAndBuildContextMenu)
 			[
 				SNew(SBox)
 				.HeightOverride(24.0f)
-				.WidthOverride(10.0f)
+				.WidthOverride(18.0f)
 				[
 					SNew(SButton)
 					.ButtonStyle(FEditorStyle::Get(), "HoverHintOnly")
 					.ContentPadding(0)
-					.OnClicked_Lambda([This]() { This->DetailsAnchor->SetIsOpen(!This->DetailsAnchor->IsOpen()); return FReply::Handled(); })
+					.OnClicked_Lambda([this]() { DetailsAnchor->SetIsOpen(!DetailsAnchor->IsOpen()); return FReply::Handled(); })
 					.VAlign(VAlign_Center)
 					.HAlign(HAlign_Center)
-					.Visibility_Lambda([This]() { return This->IsHovered() ? EVisibility::Visible : EVisibility::Hidden; })
+					.Visibility_Lambda([this]() { return IsHovered() ? EVisibility::Visible : EVisibility::Hidden; })
 					[
 						SNew(SImage)
 						.ColorAndOpacity(FSlateColor::UseForeground())
@@ -1164,9 +1328,29 @@ FReply SSequencerPlaylistItemWidget::HandleDragDetected(const FGeometry& MyGeome
 }
 
 
+void SSequencerPlaylistItemWidget::SetLoopMode(ELoopMode InLoopMode)
+{
+	if (LoopMode == InLoopMode)
+	{
+		return;
+	}
+
+	LoopMode = InLoopMode;
+
+	if (InLoopMode == ELoopMode::None)
+	{
+		USequencerPlaylistItem* Item = RowData->WeakItem.Get();
+		if (ensure(Item))
+		{
+			Item->NumLoops = 0;
+		}
+	}
+}
+
+
 EVisibility SSequencerPlaylistItemWidget::GetRowDimmingVisibility() const
 {
-	if (GetTriggerModeTransportVisibility() == EVisibility::Visible)
+	if (GetPlayModeTransportVisibility() == EVisibility::Visible)
 	{
 		return EVisibility::Visible;
 	}
@@ -1180,9 +1364,9 @@ EVisibility SSequencerPlaylistItemWidget::GetRowDimmingVisibility() const
 }
 
 
-EVisibility SSequencerPlaylistItemWidget::GetTriggerModeTransportVisibility() const
+EVisibility SSequencerPlaylistItemWidget::GetPlayModeTransportVisibility() const
 {
-	if (IsHovered() && InTriggerMode())
+	if (IsHovered() && InPlayMode())
 	{
 		return EVisibility::Visible;
 	}
@@ -1225,8 +1409,6 @@ TSharedRef<SWidget> SSequencerPlaylistItemWidget::EnsureSelectedAndBuildContextM
 
 TSharedRef<SWidget> SSequencerPlaylistItemWidget::BuildContextMenu(const TArray<UObject*>& SelectedItems)
 {
-	TSharedRef<SSequencerPlaylistItemWidget> This = SharedThis(this);
-
 	FSequencerPlaylistsModule& Module = static_cast<FSequencerPlaylistsModule&>(ISequencerPlaylistsModule::Get());
 	FMenuBuilder MenuBuilder(true, Module.GetCommandList());
 
@@ -1234,23 +1416,23 @@ TSharedRef<SWidget> SSequencerPlaylistItemWidget::BuildContextMenu(const TArray<
 	{
 		MenuBuilder.AddMenuEntry(
 			LOCTEXT("ItemContextPlay", "Play"),
-			LOCTEXT("ItemContextPlayTooltip", "Trigger playback of item"),
+			LOCTEXT("ItemContextPlayTooltip", "Begin playback of item"),
 			FSlateIcon(),
-			FUIAction(FExecuteAction::CreateLambda([This]() { This->PlayClickedDelegate.Execute(This); }))
+			FUIAction(FExecuteAction::CreateLambda([this]() { PlayClickedDelegate.Execute(*this); }))
 		);
 
 		MenuBuilder.AddMenuEntry(
 			LOCTEXT("ItemContextStop", "Stop"),
 			LOCTEXT("ItemContextStopTooltip", "Stop playback of this item"),
 			FSlateIcon(),
-			FUIAction(FExecuteAction::CreateLambda([This]() { This->StopClickedDelegate.Execute(This); }))
+			FUIAction(FExecuteAction::CreateLambda([this]() { StopClickedDelegate.Execute(*this); }))
 		);
 
 		MenuBuilder.AddMenuEntry(
 			LOCTEXT("ItemContextReset", "Reset"),
 			LOCTEXT("ItemContextResetTooltip", "Reset playback of this item"),
 			FSlateIcon(),
-			FUIAction(FExecuteAction::CreateLambda([This]() { This->ResetClickedDelegate.Execute(This); }))
+			FUIAction(FExecuteAction::CreateLambda([this]() { ResetClickedDelegate.Execute(*this); }))
 		);
 	}
 
@@ -1282,7 +1464,7 @@ TSharedRef<SWidget> SSequencerPlaylistItemWidget::BuildContextMenu(const TArray<
 			LOCTEXT("ItemContextRemove", "Remove from Playlist"),
 			LOCTEXT("ItemContextRemoveTooltip", "Remove this item from the Playlist"),
 			FSlateIcon(FEditorStyle::GetStyleSetName(), "ContentBrowser.AssetActions.Delete"),
-			FUIAction(FExecuteAction::CreateLambda([This]() { This->RemoveClickedDelegate.Execute(This); }))
+			FUIAction(FExecuteAction::CreateLambda([this]() { RemoveClickedDelegate.Execute(*this); }))
 		);
 	}
 
diff --git a/Engine/Plugins/VirtualProduction/SequencerPlaylists/Source/SequencerPlaylists/Private/SequencerPlaylistsWidgets.h b/Engine/Plugins/VirtualProduction/SequencerPlaylists/Source/SequencerPlaylists/Private/SequencerPlaylistsWidgets.h
--- a/Engine/Plugins/VirtualProduction/SequencerPlaylists/Source/SequencerPlaylists/Private/SequencerPlaylistsWidgets.h
+++ b/Engine/Plugins/VirtualProduction/SequencerPlaylists/Source/SequencerPlaylists/Private/SequencerPlaylistsWidgets.h
@@ -50,7 +50,7 @@ public:
 
 	void Construct(const FArguments& InArgs, USequencerPlaylistPlayer* InPlayer);
 
-	bool InTriggerMode() const { return bTriggerMode; }
+	bool InPlayMode() const { return bPlayMode; }
 
 private:
 	TSharedRef<SWidget> Construct_LeftToolbar();
@@ -63,7 +63,9 @@ private:
 	void RegenerateRows();
 
 	TSharedRef<SWidget> BuildOpenPlaylistMenu();
+	void OnSavePlaylist();
 	void OnSavePlaylistAs();
+	void SavePlaylist(const FString& PackageName);
 	void OnLoadPlaylist(const FAssetData& InPreset);
 	void OnNewPlaylist();
 
@@ -75,10 +77,10 @@ private:
 	FReply HandleClicked_ResetAll();
 	FReply HandleClicked_AddSequence();
 
-	FReply HandleClicked_Item_Play(TSharedPtr<SSequencerPlaylistItemWidget> ItemWidget);
-	FReply HandleClicked_Item_Stop(TSharedPtr<SSequencerPlaylistItemWidget> ItemWidget);
-	FReply HandleClicked_Item_Reset(TSharedPtr<SSequencerPlaylistItemWidget> ItemWidget);
-	FReply HandleClicked_Item_Remove(TSharedPtr<SSequencerPlaylistItemWidget> ItemWidget);
+	FReply HandleClicked_Item_Play(SSequencerPlaylistItemWidget& ItemWidget);
+	FReply HandleClicked_Item_Stop(SSequencerPlaylistItemWidget& ItemWidget);
+	FReply HandleClicked_Item_Reset(SSequencerPlaylistItemWidget& ItemWidget);
+	FReply HandleClicked_Item_Remove(SSequencerPlaylistItemWidget& ItemWidget);
 
 	bool HandleItemDetailsIsPropertyVisible(const FPropertyAndParent& PropertyAndParent);
 
@@ -88,8 +90,9 @@ private:
 
 private:
 	TWeakObjectPtr<USequencerPlaylistPlayer> WeakPlayer;
+	TWeakObjectPtr<USequencerPlaylist> WeakLoadedPlaylist;
 
-	bool bTriggerMode = false;
+	bool bPlayMode = false;
 
 	TSharedPtr<SSearchBox> SearchBox;
 	TSharedPtr<TTextFilter<const FSequencerPlaylistRowData&>> SearchTextFilter;
@@ -120,7 +123,7 @@ private:
 };
 
 
-DECLARE_DELEGATE_RetVal_OneParam(FReply, FOnClickedSequencerPlaylistItem, TSharedPtr<SSequencerPlaylistItemWidget> /*ItemWidget*/);
+DECLARE_DELEGATE_RetVal_OneParam(FReply, FOnClickedSequencerPlaylistItem, SSequencerPlaylistItemWidget& /*ItemWidget*/);
 
 
 class SSequencerPlaylistItemWidget : public SMultiColumnTableRow<TSharedPtr<FSequencerPlaylistRowData>>
@@ -130,7 +133,7 @@ class SSequencerPlaylistItemWidget : public SMultiColumnTableRow<TSharedPtr<FSeq
 	static const FText ResetItemTooltipText;
 
 	SLATE_BEGIN_ARGS(SSequencerPlaylistItemWidget) {}
-		SLATE_ATTRIBUTE(bool, TriggerMode)
+		SLATE_ATTRIBUTE(bool, PlayMode)
 		SLATE_ATTRIBUTE(bool, IsPlaying)
 
 		SLATE_EVENT(FOnClickedSequencerPlaylistItem, OnPlayClicked)
@@ -143,6 +146,13 @@ class SSequencerPlaylistItemWidget : public SMultiColumnTableRow<TSharedPtr<FSeq
 		SLATE_EVENT(FOnAcceptDrop, OnAcceptDrop)
 	SLATE_END_ARGS()
 
+	enum class ELoopMode
+	{
+		None,
+		Finite,
+		//Infinite,
+	};
+
 public:
 	void Construct(const FArguments& InArgs, TSharedPtr<FSequencerPlaylistRowData> InRowData, const TSharedRef<STableViewBase>& OwnerTableView);
 
@@ -164,10 +174,11 @@ public:
 private:
 	FReply HandleDragDetected(const FGeometry& MyGeometry, const FPointerEvent& MouseEvent);
 
-	bool InTriggerMode() const { return TriggerMode.Get(); }
+	bool InPlayMode() const { return PlayMode.Get(); }
+	void SetLoopMode(ELoopMode InLoopMode);
 
 	EVisibility GetRowDimmingVisibility() const;
-	EVisibility GetTriggerModeTransportVisibility() const;
+	EVisibility GetPlayModeTransportVisibility() const;
 
 	TSharedRef<SWidget> EnsureSelectedAndBuildContextMenu();
 	TSharedRef<SWidget> BuildContextMenu(const TArray<UObject*>& SelectedItems);
@@ -176,7 +187,9 @@ private:
 	TSharedPtr<FSequencerPlaylistRowData> RowData;
 	TSharedPtr<SMenuAnchor> DetailsAnchor;
 
-	TAttribute<bool> TriggerMode;
+	ELoopMode LoopMode;
+
+	TAttribute<bool> PlayMode;
 	TAttribute<bool> IsPlaying;
 
 	FOnClickedSequencerPlaylistItem PlayClickedDelegate;
