Sequencer Playlists: Updated tooltips, plugin description, replace plugin wizard placeholder icon, mark plugin beta instead of experimental. Mark Sequence property BlueprintReadWrite.

diff --git a/Engine/Plugins/VirtualProduction/SequencerPlaylists/SequencerPlaylists.uplugin b/Engine/Plugins/VirtualProduction/SequencerPlaylists/SequencerPlaylists.uplugin
--- a/Engine/Plugins/VirtualProduction/SequencerPlaylists/SequencerPlaylists.uplugin
+++ b/Engine/Plugins/VirtualProduction/SequencerPlaylists/SequencerPlaylists.uplugin
@@ -3,7 +3,7 @@
 	"Version": 1,
 	"VersionName": "1.0",
 	"FriendlyName": "Playlists",
-	"Description": "",
+	"Description": "Sequencer Playlists allow users to prepare, queue, and trigger level sequences on the fly during a virtual production session, providing increased flexibility and agility when interacting with animation on set.",
 	"Category": "Virtual Production",
 	"CreatedBy": "Epic Games, Inc.",
 	"CreatedByURL": "http://epicgames.com",
@@ -11,8 +11,8 @@
 	"MarketplaceURL": "",
 	"SupportURL": "",
 	"CanContainContent": false,
-	"IsBetaVersion": false,
-	"IsExperimentalVersion": true,
+	"IsBetaVersion": true,
+	"IsExperimentalVersion": false,
 	"Installed": false,
 	"Plugins": [
 		{
diff --git a/Engine/Plugins/VirtualProduction/SequencerPlaylists/Source/SequencerPlaylists/Private/SequencerPlaylistItem_Sequence.h b/Engine/Plugins/VirtualProduction/SequencerPlaylists/Source/SequencerPlaylists/Private/SequencerPlaylistItem_Sequence.h
--- a/Engine/Plugins/VirtualProduction/SequencerPlaylists/Source/SequencerPlaylists/Private/SequencerPlaylistItem_Sequence.h
+++ b/Engine/Plugins/VirtualProduction/SequencerPlaylists/Source/SequencerPlaylists/Private/SequencerPlaylistItem_Sequence.h
@@ -22,7 +22,7 @@ class USequencerPlaylistItem_Sequence : public USequencerPlaylistItem
 	FText GetDisplayName() override;
 
 public:
-	UPROPERTY(EditAnywhere, Category="SequencerPlaylists", meta=(NoResetToDefault))
+	UPROPERTY(EditAnywhere, BlueprintReadWrite, Category="SequencerPlaylists", meta=(NoResetToDefault))
 	TObjectPtr<ULevelSequence> Sequence;
 };
 
diff --git a/Engine/Plugins/VirtualProduction/SequencerPlaylists/Source/SequencerPlaylists/Private/SequencerPlaylistsWidgets.cpp b/Engine/Plugins/VirtualProduction/SequencerPlaylists/Source/SequencerPlaylists/Private/SequencerPlaylistsWidgets.cpp
--- a/Engine/Plugins/VirtualProduction/SequencerPlaylists/Source/SequencerPlaylists/Private/SequencerPlaylistsWidgets.cpp
+++ b/Engine/Plugins/VirtualProduction/SequencerPlaylists/Source/SequencerPlaylists/Private/SequencerPlaylistsWidgets.cpp
@@ -188,7 +188,7 @@ TSharedRef<SWidget> SSequencerPlaylistPanel::Construct_LeftToolbar()
 			FUIAction(FExecuteAction::CreateSP(this, &SSequencerPlaylistPanel::OnNewPlaylist)),
 			NAME_None,
 			LOCTEXT("NewPlaylist", "New Playlist"),
-			LOCTEXT("NewPlaylistToolTip", "New Playlist"),
+			LOCTEXT("NewPlaylistTooltip", "New Playlist"),
 			FSlateIcon(FSequencerPlaylistsStyle::Get().GetStyleSetName(), "SequencerPlaylists.NewPlaylist"));
 
 		ToolBarBuilder.AddToolBarButton(
@@ -204,22 +204,22 @@ TSharedRef<SWidget> SSequencerPlaylistPanel::Construct_LeftToolbar()
 				FMenuBuilder MenuBuilder(true, nullptr);
 				MenuBuilder.AddMenuEntry(
 					LOCTEXT("SavePlaylistAs", "Save Playlist As..."),
-					LOCTEXT("SavePlaylistAsTooltip", "Save Playlist As..."),
+					LOCTEXT("SavePlaylistAsTooltip", "Save Playlist as..."),
 					FSlateIcon(),
 					FUIAction(FExecuteAction::CreateSP(this, &SSequencerPlaylistPanel::OnSavePlaylistAs))
 				);
 				return MenuBuilder.MakeWidget();
 			}),
 			LOCTEXT("SavePlaylistOptions", "Save Playlist Options"),
-			LOCTEXT("SavePlaylistOptionsTooltip", "Save Playlist Options"),
+			LOCTEXT("SavePlaylistOptionsTooltip", "Save Playlist options"),
 			TAttribute<FSlateIcon>(),
 			true);
 
 		ToolBarBuilder.AddComboButton(
 			FUIAction(),
 			FOnGetContent::CreateSP(this, &SSequencerPlaylistPanel::BuildOpenPlaylistMenu),
-			LOCTEXT("OpenPlaylist", "Open Playlist"),
-			LOCTEXT("OpenPlaylistToolTip", "Open Playlist"),
+			LOCTEXT("OpenPlaylist", "Load Playlist"),
+			LOCTEXT("OpenPlaylistTooltip", "Load Playlist"),
 			FSlateIcon(FSequencerPlaylistsStyle::Get().GetStyleSetName(), "SequencerPlaylists.OpenPlaylist"));
 	}
 	ToolBarBuilder.EndSection();
@@ -250,7 +250,7 @@ TSharedRef<SWidget> SSequencerPlaylistPanel::Construct_RightToolbar()
 			),
 			NAME_None,
 			LOCTEXT("TogglePlayMode", "Toggle Play Mode"),
-			LOCTEXT("TogglePlayModeToolTip", "Enables or disables Play Mode"),
+			LOCTEXT("TogglePlayModeTooltip", "Enable or disable Play Mode. Play Mode disables editing and provides larger play buttons for easier playback."),
 			FSlateIcon(FSequencerPlaylistsStyle::Get().GetStyleSetName(), "SequencerPlaylists.PlayMode"),
 			EUserInterfaceActionType::ToggleButton);
 	}
@@ -274,7 +274,7 @@ TSharedRef<SWidget> SSequencerPlaylistPanel::Construct_Transport()
 			.HAlign(HAlign_Center)
 			.VAlign(VAlign_Center)
 			.OnClicked(this, &SSequencerPlaylistPanel::HandleClicked_PlayAll)
-			.ToolTipText(LOCTEXT("PlayAllButtonToolTip", "Play all items"))
+			.ToolTipText(LOCTEXT("PlayAllButtonTooltip", "Play all items simultaneously."))
 			[
 				SNew(SImage)
 				.Image(FSequencerPlaylistsStyle::Get().GetBrush("SequencerPlaylists.Play"))
@@ -290,7 +290,7 @@ TSharedRef<SWidget> SSequencerPlaylistPanel::Construct_Transport()
 			.HAlign(HAlign_Center)
 			.VAlign(VAlign_Center)
 			.OnClicked(this, &SSequencerPlaylistPanel::HandleClicked_StopAll)
-			.ToolTipText(LOCTEXT("StopAllButtonToolTip", "Stop all items"))
+			.ToolTipText(LOCTEXT("StopAllButtonTooltip", "Stop all currently playing items."))
 			[
 				SNew(SImage)
 				.Image(FSequencerPlaylistsStyle::Get().GetBrush("SequencerPlaylists.Stop"))
@@ -306,7 +306,7 @@ TSharedRef<SWidget> SSequencerPlaylistPanel::Construct_Transport()
 			.HAlign(HAlign_Center)
 			.VAlign(VAlign_Center)
 			.OnClicked(this, &SSequencerPlaylistPanel::HandleClicked_ResetAll)
-			.ToolTipText(LOCTEXT("ResetAllButtonToolTip", "Reset all items"))
+			.ToolTipText(LOCTEXT("ResetAllButtonTooltip", "Stop all currently playing items and re-hold, if specified."))
 			[
 				SNew(SImage)
 				.Image(FSequencerPlaylistsStyle::Get().GetBrush("SequencerPlaylists.Reset"))
@@ -328,6 +328,7 @@ TSharedRef<SWidget> SSequencerPlaylistPanel::Construct_AddSearchRow()
 			.OnClicked(this, &SSequencerPlaylistPanel::HandleClicked_AddSequence)
 			.Icon(FAppStyle::Get().GetBrush("Icons.Plus"))
 			.Text(LOCTEXT("AddItemButton", "Item"))
+			.ToolTipText(LOCTEXT("AddItemButtonTooltip", "Add a new level sequence Playlist item"))
 		]
 		+ SHorizontalBox::Slot()
 		.FillWidth(1.0f)
@@ -919,9 +920,9 @@ FSequencerPlaylistItemDragDropOp::~FSequencerPlaylistItemDragDropOp()
 }
 
 
-const FText SSequencerPlaylistItemWidget::PlayItemTooltipText(LOCTEXT("PlayItemButtonToolTip", "Play this item"));
-const FText SSequencerPlaylistItemWidget::StopItemTooltipText(LOCTEXT("StopItemButtonToolTip", "Stop playback of this item"));
-const FText SSequencerPlaylistItemWidget::ResetItemTooltipText(LOCTEXT("ResetItemButtonToolTip", "Reset this item"));
+const FText SSequencerPlaylistItemWidget::PlayItemTooltipText(LOCTEXT("PlayItemTooltip", "Play just this item."));
+const FText SSequencerPlaylistItemWidget::StopItemTooltipText(LOCTEXT("StopItemTooltip", "Stop all running sequences created by this item."));
+const FText SSequencerPlaylistItemWidget::ResetItemTooltipText(LOCTEXT("ResetItemTooltip", "Stop all running sequences created by this item and re-hold, if specified."));
 
 
 void SSequencerPlaylistItemWidget::Construct(const FArguments& InArgs, TSharedPtr<FSequencerPlaylistRowData> InRowData, const TSharedRef<STableViewBase>& OwnerTableView)
@@ -1160,7 +1161,7 @@ TSharedRef<SWidget> SSequencerPlaylistItemWidget::GenerateWidgetForColumn(const
 		TSharedRef<SWidget> HoldToggle = SNew(SCheckBox)
 			.Padding(FMargin(4.0f, 2.0f))
 			.HAlign(HAlign_Center)
-			.ToolTipText(LOCTEXT("ToggleHoldToolTip", "Enables or disables Hold At First Frame"))
+			.ToolTipText(LOCTEXT("ToggleHoldTooltip", "Enables or disables hold. Hold will infinitely hold the first frame of this item until manually played. Items are put into a hold state at the start of a take, or manually by hitting \"Reset.\""))
 			.Style(FEditorStyle::Get(), "ToggleButtonCheckbox")
 			.IsChecked_Lambda([WeakItem]() { return (WeakItem.IsValid() && WeakItem->bHoldAtFirstFrame) ? ECheckBoxState::Checked : ECheckBoxState::Unchecked; })
 			.OnCheckStateChanged_Lambda([WeakItem](ECheckBoxState InState) {
@@ -1254,6 +1255,7 @@ TSharedRef<SWidget> SSequencerPlaylistItemWidget::GenerateWidgetForColumn(const
 				[
 					SNew(SNumericEntryBox<int32>)
 					.AllowSpin(true)
+					.ToolTipText(LOCTEXT("LoopCountTooltip", "Number of times to loop before stopping. A value of 1 will result in a sequence playing twice before stopping."))
 					.Value(TAttribute<TOptional<int32>>::CreateLambda([WeakItem]() { return WeakItem.IsValid() ? WeakItem->NumLoops : TOptional<int32>(); }))
 					.Visibility_Lambda([this]() { return LoopMode == ELoopMode::Finite ? EVisibility::Visible : EVisibility::Hidden; })
 					.OnValueCommitted_Lambda([WeakItem](int32 NewValue, ETextCommit::Type) {
@@ -1416,21 +1418,21 @@ TSharedRef<SWidget> SSequencerPlaylistItemWidget::BuildContextMenu(const TArray<
 	{
 		MenuBuilder.AddMenuEntry(
 			LOCTEXT("ItemContextPlay", "Play"),
-			LOCTEXT("ItemContextPlayTooltip", "Begin playback of item"),
+			PlayItemTooltipText,
 			FSlateIcon(),
 			FUIAction(FExecuteAction::CreateLambda([this]() { PlayClickedDelegate.Execute(*this); }))
 		);
 
 		MenuBuilder.AddMenuEntry(
 			LOCTEXT("ItemContextStop", "Stop"),
-			LOCTEXT("ItemContextStopTooltip", "Stop playback of this item"),
+			StopItemTooltipText,
 			FSlateIcon(),
 			FUIAction(FExecuteAction::CreateLambda([this]() { StopClickedDelegate.Execute(*this); }))
 		);
 
 		MenuBuilder.AddMenuEntry(
 			LOCTEXT("ItemContextReset", "Reset"),
-			LOCTEXT("ItemContextResetTooltip", "Reset playback of this item"),
+			ResetItemTooltipText,
 			FSlateIcon(),
 			FUIAction(FExecuteAction::CreateLambda([this]() { ResetClickedDelegate.Execute(*this); }))
 		);
diff --git a/Engine/Plugins/VirtualProduction/SequencerPlaylists/Source/SequencerPlaylists/Public/SequencerPlaylistItem.h b/Engine/Plugins/VirtualProduction/SequencerPlaylists/Source/SequencerPlaylists/Public/SequencerPlaylistItem.h
--- a/Engine/Plugins/VirtualProduction/SequencerPlaylists/Source/SequencerPlaylists/Public/SequencerPlaylistItem.h
+++ b/Engine/Plugins/VirtualProduction/SequencerPlaylists/Source/SequencerPlaylists/Public/SequencerPlaylistItem.h
@@ -19,15 +19,17 @@ public:
 	virtual FText GetDisplayName() PURE_VIRTUAL(USequencerPlaylistItem::GetDisplayName, return FText::GetEmpty(); )
 
 public:
+	/** Number of frames by which to clip the in point of sections played from this item. Will also affect the first frame for hold. */
 	UPROPERTY(EditAnywhere, BlueprintReadWrite, Category="SequencerPlaylists", meta=(NoResetToDefault))
 	int32 StartFrameOffset;
 
+	/** Number of frames by which to clip the out point of sections played from this item. */
 	UPROPERTY(EditAnywhere, BlueprintReadWrite, Category="SequencerPlaylists", meta=(NoResetToDefault))
 	int32 EndFrameOffset;
 
 	/**
 	 * If true, the sequence will be inserted immediately on recording start and any time Reset()
-	 * is called, paused at the first frame indefinitely until either triggered or stopped.
+	 * is called, paused at the first frame indefinitely until subsequently played or stopped.
 	 */
 	UPROPERTY(EditAnywhere, BlueprintReadWrite, Category="SequencerPlaylists", meta=(NoResetToDefault))
 	bool bHoldAtFirstFrame;
@@ -36,9 +38,11 @@ public:
 	UPROPERTY(EditAnywhere, BlueprintReadWrite, Category="SequencerPlaylists", meta=(NoResetToDefault))
 	int32 NumLoops;
 
+	/** Speed multiplier at which to play sections created by this item. A value of 1 will result in playback at the original speed. */
 	UPROPERTY(EditAnywhere, BlueprintReadWrite, Category="SequencerPlaylists", meta=(NoResetToDefault))
 	float PlaybackSpeed = 1.0f;
 
+	/** Disables playback of this item and prevents it from entering a hold state. */
 	UPROPERTY(EditAnywhere, BlueprintReadWrite, Category="SequencerPlaylists", meta=(NoResetToDefault))
 	bool bMute;
 };
