Sequencer Playlists: Add missing calls to NotifyMovieSceneDataChanged for single item play/stop/reset actions. Tidy transaction descriptions.

diff --git a/Engine/Plugins/VirtualProduction/SequencerPlaylists/Source/SequencerPlaylists/Private/SequencerPlaylistPlayer.cpp b/Engine/Plugins/VirtualProduction/SequencerPlaylists/Source/SequencerPlaylists/Private/SequencerPlaylistPlayer.cpp
--- a/Engine/Plugins/VirtualProduction/SequencerPlaylists/Source/SequencerPlaylists/Private/SequencerPlaylistPlayer.cpp
+++ b/Engine/Plugins/VirtualProduction/SequencerPlaylists/Source/SequencerPlaylists/Private/SequencerPlaylistPlayer.cpp
@@ -66,8 +66,17 @@ bool USequencerPlaylistPlayer::PlayItem(USequencerPlaylistItem* Item)
 
 	EnterUnboundedPlayIfNotRecording();
 
-	FScopedTransaction Transaction(FText::Format(LOCTEXT("PlayItemTransaction", "Trigger playback of {0}"), Item->GetDisplayName()));
-	return GetCheckedItemPlayer(Item)->Play(Item);
+	FScopedTransaction Transaction(FText::Format(LOCTEXT("PlayItemTransaction", "Begin playback of {0}"), Item->GetDisplayName()));
+
+	if (GetCheckedItemPlayer(Item)->Play(Item))
+	{
+		GetSequencer()->NotifyMovieSceneDataChanged(EMovieSceneDataChangeType::RefreshAllImmediately);
+		return true;
+	}
+	else
+	{
+		return false;
+	}
 }
 
 
@@ -79,7 +88,16 @@ bool USequencerPlaylistPlayer::StopItem(USequencerPlaylistItem* Item)
 	}
 
 	FScopedTransaction Transaction(FText::Format(LOCTEXT("StopItemTransaction", "Stop playback of {0}"), Item->GetDisplayName()));
-	return GetCheckedItemPlayer(Item)->Stop(Item);
+
+	if (GetCheckedItemPlayer(Item)->Stop(Item))
+	{
+		GetSequencer()->NotifyMovieSceneDataChanged(EMovieSceneDataChangeType::RefreshAllImmediately);
+		return true;
+	}
+	else
+	{
+		return false;
+	}
 }
 
 
@@ -91,7 +109,16 @@ bool USequencerPlaylistPlayer::ResetItem(USequencerPlaylistItem* Item)
 	}
 
 	FScopedTransaction Transaction(FText::Format(LOCTEXT("ResetItemTransaction", "Reset playback of {0}"), Item->GetDisplayName()));
-	return GetCheckedItemPlayer(Item)->Reset(Item);
+
+	if (GetCheckedItemPlayer(Item)->Reset(Item))
+	{
+		GetSequencer()->NotifyMovieSceneDataChanged(EMovieSceneDataChangeType::RefreshAllImmediately);
+		return true;
+	}
+	else
+	{
+		return false;
+	}
 }
 
 
@@ -235,7 +262,7 @@ bool USequencerPlaylistPlayer::PlayAll()
 
 	bool bAnyChange = false;
 
-	FScopedTransaction Transaction(LOCTEXT("PlayAllTransaction", "Trigger playback of all items"));
+	FScopedTransaction Transaction(LOCTEXT("PlayAllTransaction", "Begin playback of all Playlist items"));
 	for (USequencerPlaylistItem* Item : Playlist->Items)
 	{
 		if (Item->bMute)
@@ -264,7 +291,7 @@ bool USequencerPlaylistPlayer::StopAll()
 
 	bool bAnyChange = false;
 
-	FScopedTransaction Transaction(LOCTEXT("StopAllTransaction", "Stop playback of all items"));
+	FScopedTransaction Transaction(LOCTEXT("StopAllTransaction", "Stop playback of all Playlist items"));
 	for (USequencerPlaylistItem* Item : Playlist->Items)
 	{
 		bAnyChange |= GetCheckedItemPlayer(Item)->Stop(Item);
@@ -299,7 +326,7 @@ bool USequencerPlaylistPlayer::ResetAll()
 
 	bool bAnyChange = false;
 
-	FScopedTransaction Transaction(LOCTEXT("ResetAllTransaction", "Reset playback of all items"));
+	FScopedTransaction Transaction(LOCTEXT("ResetAllTransaction", "Reset playback of all Playlist items"));
 	for (USequencerPlaylistItem* Item : Playlist->Items)
 	{
 		if (Item->bMute)
