Sequencer Playlists: Tooltips and other minor fixes.

- Changed start/end offset controls from single property views to SNumericEntryBox to enable tooltips to be displayed correctly.
- Bounded offset ranges to sequence length.
- Fixed an issue where loop count did not update during spin inputs.
- Changed tab title from "Playlist" back to "Playlists" (mistake in latest mockups).
- Removed `meta=(NoResetToDefault)` workarounds from USequencerPlaylistItem properties, which is no longer necessary now that we've moved away from single property views.
- Moved some of the recently-added "unbounded play" code within transaction scope.
- Querying IsPlaying() will no longer open Sequencer if it's not already open.

diff --git a/Engine/Plugins/VirtualProduction/SequencerPlaylists/Source/SequencerPlaylists/Private/SequencerPlaylistPlayer.cpp b/Engine/Plugins/VirtualProduction/SequencerPlaylists/Source/SequencerPlaylists/Private/SequencerPlaylistPlayer.cpp
--- a/Engine/Plugins/VirtualProduction/SequencerPlaylists/Source/SequencerPlaylists/Private/SequencerPlaylistPlayer.cpp
+++ b/Engine/Plugins/VirtualProduction/SequencerPlaylists/Source/SequencerPlaylists/Private/SequencerPlaylistPlayer.cpp
@@ -64,10 +64,10 @@ bool USequencerPlaylistPlayer::PlayItem(USequencerPlaylistItem* Item)
 		return false;
 	}
 
-	EnterUnboundedPlayIfNotRecording();
-
 	FScopedTransaction Transaction(FText::Format(LOCTEXT("PlayItemTransaction", "Begin playback of {0}"), Item->GetDisplayName()));
 
+	EnterUnboundedPlayIfNotRecording();
+
 	if (GetCheckedItemPlayer(Item)->Play(Item))
 	{
 		GetSequencer()->NotifyMovieSceneDataChanged(EMovieSceneDataChangeType::RefreshAllImmediately);
@@ -129,6 +129,11 @@ bool USequencerPlaylistPlayer::IsPlaying(USequencerPlaylistItem* Item)
 		return false;
 	}
 
+	if (!WeakSequencer.IsValid())
+	{
+		return false;
+	}
+
 	return GetCheckedItemPlayer(Item)->IsPlaying(Item);
 }
 
@@ -258,11 +263,12 @@ bool USequencerPlaylistPlayer::PlayAll()
 		return false;
 	}
 
+	FScopedTransaction Transaction(LOCTEXT("PlayAllTransaction", "Begin playback of all Playlist items"));
+
 	EnterUnboundedPlayIfNotRecording();
 
 	bool bAnyChange = false;
 
-	FScopedTransaction Transaction(LOCTEXT("PlayAllTransaction", "Begin playback of all Playlist items"));
 	for (USequencerPlaylistItem* Item : Playlist->Items)
 	{
 		if (Item->bMute)
diff --git a/Engine/Plugins/VirtualProduction/SequencerPlaylists/Source/SequencerPlaylists/Private/SequencerPlaylistsModule.cpp b/Engine/Plugins/VirtualProduction/SequencerPlaylists/Source/SequencerPlaylists/Private/SequencerPlaylistsModule.cpp
--- a/Engine/Plugins/VirtualProduction/SequencerPlaylists/Source/SequencerPlaylists/Private/SequencerPlaylistsModule.cpp
+++ b/Engine/Plugins/VirtualProduction/SequencerPlaylists/Source/SequencerPlaylists/Private/SequencerPlaylistsModule.cpp
@@ -103,7 +103,6 @@ TSharedRef<SDockTab> FSequencerPlaylistsModule::OnSpawnPluginTab(const FSpawnTab
 {
 	return SNew(SDockTab)
 		.TabRole(ETabRole::NomadTab)
-		.Label(LOCTEXT("TabLabel", "Playlist"))
 		[
 			SNew(SSequencerPlaylistPanel, DefaultPlayer.Get())
 		];
diff --git a/Engine/Plugins/VirtualProduction/SequencerPlaylists/Source/SequencerPlaylists/Private/SequencerPlaylistsWidgets.cpp b/Engine/Plugins/VirtualProduction/SequencerPlaylists/Source/SequencerPlaylists/Private/SequencerPlaylistsWidgets.cpp
--- a/Engine/Plugins/VirtualProduction/SequencerPlaylists/Source/SequencerPlaylists/Private/SequencerPlaylistsWidgets.cpp
+++ b/Engine/Plugins/VirtualProduction/SequencerPlaylists/Source/SequencerPlaylists/Private/SequencerPlaylistsWidgets.cpp
@@ -18,8 +18,10 @@
 #include "IContentBrowserSingleton.h"
 #include "ISinglePropertyView.h"
 #include "LevelEditor.h"
+#include "LevelSequence.h"
 #include "Misc/FileHelper.h"
 #include "Misc/TextFilter.h"
+#include "MovieScene.h"
 #include "ScopedTransaction.h"
 #include "SlateOptMacros.h"
 #include "SPositiveActionButton.h"
@@ -51,7 +53,6 @@ const FName SSequencerPlaylistPanel::ColumnName_Loop(TEXT("Loop"));
 const FName SSequencerPlaylistPanel::ColumnName_HoverDetails(TEXT("HoverDetails"));
 
 
-BEGIN_SLATE_FUNCTION_BUILD_OPTIMIZATION
 void SSequencerPlaylistPanel::Construct(const FArguments& InArgs, USequencerPlaylistPlayer* InPlayer)
 {
 	check(InPlayer);
@@ -172,10 +173,8 @@ void SSequencerPlaylistPanel::Construct(const FArguments& InArgs, USequencerPlay
 
 	RegenerateRows();
 }
-END_SLATE_FUNCTION_BUILD_OPTIMIZATION
 
 
-BEGIN_SLATE_FUNCTION_BUILD_OPTIMIZATION
 TSharedRef<SWidget> SSequencerPlaylistPanel::Construct_LeftToolbar()
 {
 	FSequencerPlaylistsModule& Module = static_cast<FSequencerPlaylistsModule&>(ISequencerPlaylistsModule::Get());
@@ -226,10 +225,8 @@ TSharedRef<SWidget> SSequencerPlaylistPanel::Construct_LeftToolbar()
 
 	return ToolBarBuilder.MakeWidget();
 }
-END_SLATE_FUNCTION_BUILD_OPTIMIZATION
 
 
-BEGIN_SLATE_FUNCTION_BUILD_OPTIMIZATION
 TSharedRef<SWidget> SSequencerPlaylistPanel::Construct_RightToolbar()
 {
 	FSequencerPlaylistsModule& Module = static_cast<FSequencerPlaylistsModule&>(ISequencerPlaylistsModule::Get());
@@ -258,10 +255,8 @@ TSharedRef<SWidget> SSequencerPlaylistPanel::Construct_RightToolbar()
 
 	return ToolBarBuilder.MakeWidget();
 }
-END_SLATE_FUNCTION_BUILD_OPTIMIZATION
 
 
-BEGIN_SLATE_FUNCTION_BUILD_OPTIMIZATION
 TSharedRef<SWidget> SSequencerPlaylistPanel::Construct_Transport()
 {
 	return SNew(SHorizontalBox)
@@ -314,10 +309,8 @@ TSharedRef<SWidget> SSequencerPlaylistPanel::Construct_Transport()
 			]
 		];
 }
-END_SLATE_FUNCTION_BUILD_OPTIMIZATION
 
 
-BEGIN_SLATE_FUNCTION_BUILD_OPTIMIZATION
 TSharedRef<SWidget> SSequencerPlaylistPanel::Construct_AddSearchRow()
 {
 	return SNew(SHorizontalBox)
@@ -340,10 +333,8 @@ TSharedRef<SWidget> SSequencerPlaylistPanel::Construct_AddSearchRow()
 			.DelayChangeNotificationsWhileTyping(true)
 		];
 }
-END_SLATE_FUNCTION_BUILD_OPTIMIZATION
 
 
-BEGIN_SLATE_FUNCTION_BUILD_OPTIMIZATION
 TSharedRef<SWidget> SSequencerPlaylistPanel::Construct_ItemListView()
 {
 	return SAssignNew(ItemListView, SListView<TSharedPtr<FSequencerPlaylistRowData>>)
@@ -398,7 +389,6 @@ TSharedRef<SWidget> SSequencerPlaylistPanel::Construct_ItemListView()
 				.VAlignCell(VAlign_Center)
 		);
 }
-END_SLATE_FUNCTION_BUILD_OPTIMIZATION
 
 
 USequencerPlaylist* SSequencerPlaylistPanel::GetCheckedPlaylist()
@@ -1047,7 +1037,6 @@ void SSequencerPlaylistItemWidget::ConstructChildren(ETableViewMode::Type InOwne
 }
 
 
-BEGIN_SLATE_FUNCTION_BUILD_OPTIMIZATION
 TSharedRef<SWidget> SSequencerPlaylistItemWidget::GenerateWidgetForColumn(const FName& ColumnName)
 {
 	USequencerPlaylistItem* Item = GetItem();
@@ -1139,20 +1128,46 @@ TSharedRef<SWidget> SSequencerPlaylistItemWidget::GenerateWidgetForColumn(const
 
 	if (ColumnName == SSequencerPlaylistPanel::ColumnName_Offset)
 	{
-		TSharedPtr<ISinglePropertyView> StartOffsetPropView = PropertyEditorModule.CreateSingleProperty(
-			Item, GET_MEMBER_NAME_CHECKED(USequencerPlaylistItem, StartFrameOffset), SinglePropParams);
-
-		TSharedPtr<ISinglePropertyView> EndOffsetPropView = PropertyEditorModule.CreateSingleProperty(
-			Item, GET_MEMBER_NAME_CHECKED(USequencerPlaylistItem, EndFrameOffset), SinglePropParams);
-
 		return SNew(SVerticalBox)
 			+ SVerticalBox::Slot()
 			[
-				StartOffsetPropView.ToSharedRef()
+				SNew(SNumericEntryBox<int32>)
+				.AllowSpin(true)
+				.ToolTipText(LOCTEXT("StartFrameOffsetTooltip", "Number of frames by which to clip the in point of sections played from this item. Will also affect the first frame for hold."))
+				.MaxSliderValue(this, &SSequencerPlaylistItemWidget::GetItemLengthDisplayFrames)
+				.Value(TAttribute<TOptional<int32>>::CreateLambda([WeakItem]() { return WeakItem.IsValid() ? WeakItem->StartFrameOffset : TOptional<int32>(); }))
+				.OnValueChanged_Lambda([WeakItem](int32 NewValue) {
+					if (WeakItem.IsValid())
+					{
+						WeakItem->StartFrameOffset = NewValue;
+					}
+				})
+				.OnValueCommitted_Lambda([WeakItem](int32 NewValue, ETextCommit::Type) {
+					if (WeakItem.IsValid())
+					{
+						WeakItem->StartFrameOffset = NewValue;
+					}
+				})
 			]
 			+ SVerticalBox::Slot()
 			[
-				EndOffsetPropView.ToSharedRef()
+				SNew(SNumericEntryBox<int32>)
+				.AllowSpin(true)
+				.ToolTipText(LOCTEXT("EndFrameOffsetTooltip", "Number of frames by which to clip the out point of sections played from this item."))
+				.MaxSliderValue(this, &SSequencerPlaylistItemWidget::GetItemLengthDisplayFrames)
+				.Value(TAttribute<TOptional<int32>>::CreateLambda([WeakItem]() { return WeakItem.IsValid() ? WeakItem->EndFrameOffset : TOptional<int32>(); }))
+				.OnValueChanged_Lambda([WeakItem](int32 NewValue) {
+					if (WeakItem.IsValid())
+					{
+						WeakItem->EndFrameOffset = NewValue;
+					}
+				})
+				.OnValueCommitted_Lambda([WeakItem](int32 NewValue, ETextCommit::Type) {
+					if (WeakItem.IsValid())
+					{
+						WeakItem->EndFrameOffset = NewValue;
+					}
+				})
 			];
 	}
 
@@ -1161,7 +1176,7 @@ TSharedRef<SWidget> SSequencerPlaylistItemWidget::GenerateWidgetForColumn(const
 		TSharedRef<SWidget> HoldToggle = SNew(SCheckBox)
 			.Padding(FMargin(4.0f, 2.0f))
 			.HAlign(HAlign_Center)
-			.ToolTipText(LOCTEXT("ToggleHoldTooltip", "Enables or disables hold. Hold will infinitely hold the first frame of this item until manually played. Items are put into a hold state at the start of a take, or manually by hitting \"Reset.\""))
+			.ToolTipText(LOCTEXT("ToggleHoldTooltip", "Enable or disable hold. Hold will infinitely hold the first frame of this item until manually played. Items are put into a hold state at the start of a take, or manually by hitting \"Reset.\""))
 			.Style(FEditorStyle::Get(), "ToggleButtonCheckbox")
 			.IsChecked_Lambda([WeakItem]() { return (WeakItem.IsValid() && WeakItem->bHoldAtFirstFrame) ? ECheckBoxState::Checked : ECheckBoxState::Unchecked; })
 			.OnCheckStateChanged_Lambda([WeakItem](ECheckBoxState InState) {
@@ -1193,11 +1208,12 @@ TSharedRef<SWidget> SSequencerPlaylistItemWidget::GenerateWidgetForColumn(const
 		TSharedRef<SWidget> LoopModeCombo = SNew(SComboButton)
 			.ContentPadding(0)
 			.ComboButtonStyle(FAppStyle::Get(), "SimpleComboButton")
+			.ToolTipText(LOCTEXT("LoopModeComboTooltip", "Toggle between looping a specified number of times or not looping."))
 			.OnGetMenuContent_Lambda([this]() {
 				FMenuBuilder MenuBuilder(true, nullptr);
 				MenuBuilder.AddMenuEntry(
-					LOCTEXT("LoopModeNone", "Don't Loop"),
-					LOCTEXT("LoopModeNoneTooltip", "Don't Loop"),
+					LOCTEXT("LoopModeNone", "No Loop"),
+					LOCTEXT("LoopModeNoneTooltip", "No Loop"),
 					FSlateIcon(FSequencerPlaylistsStyle::Get().GetStyleSetName(), "SequencerPlaylists.Loop.Disabled"),
 					FUIAction(
 						FExecuteAction::CreateLambda([this]() { SetLoopMode(ELoopMode::None); }),
@@ -1208,8 +1224,9 @@ TSharedRef<SWidget> SSequencerPlaylistItemWidget::GenerateWidgetForColumn(const
 					EUserInterfaceActionType::Check
 				);
 				MenuBuilder.AddMenuEntry(
-					LOCTEXT("LoopModeFinite", "Loop N Times"),
-					LOCTEXT("LoopModeFiniteTooltip", "Loop N Times"),
+					// U+1D62F = "Mathematical Sans-Serif Italic Small N"
+					LOCTEXT("LoopModeFinite", "Loop [\U0001D62F] Times"),
+					LOCTEXT("LoopModeFiniteTooltip", "Loop [\U0001D62F] Times"),
 					FSlateIcon(FSequencerPlaylistsStyle::Get().GetStyleSetName(), "SequencerPlaylists.Loop.Finite"),
 					FUIAction(
 						FExecuteAction::CreateLambda([this]() { SetLoopMode(ELoopMode::Finite); }),
@@ -1258,6 +1275,12 @@ TSharedRef<SWidget> SSequencerPlaylistItemWidget::GenerateWidgetForColumn(const
 					.ToolTipText(LOCTEXT("LoopCountTooltip", "Number of times to loop before stopping. A value of 1 will result in a sequence playing twice before stopping."))
 					.Value(TAttribute<TOptional<int32>>::CreateLambda([WeakItem]() { return WeakItem.IsValid() ? WeakItem->NumLoops : TOptional<int32>(); }))
 					.Visibility_Lambda([this]() { return LoopMode == ELoopMode::Finite ? EVisibility::Visible : EVisibility::Hidden; })
+					.OnValueChanged_Lambda([WeakItem](int32 NewValue) {
+						if (WeakItem.IsValid())
+						{
+							WeakItem->NumLoops = NewValue;
+						}
+					})
 					.OnValueCommitted_Lambda([WeakItem](int32 NewValue, ETextCommit::Type) {
 						if (WeakItem.IsValid())
 						{
@@ -1301,7 +1324,6 @@ TSharedRef<SWidget> SSequencerPlaylistItemWidget::GenerateWidgetForColumn(const
 
 	return SNullWidget::NullWidget;
 }
-END_SLATE_FUNCTION_BUILD_OPTIMIZATION
 
 
 FReply SSequencerPlaylistItemWidget::OnMouseButtonUp(const FGeometry& MyGeometry, const FPointerEvent& MouseEvent)
@@ -1377,6 +1399,27 @@ EVisibility SSequencerPlaylistItemWidget::GetPlayModeTransportVisibility() const
 }
 
 
+TOptional<int32> SSequencerPlaylistItemWidget::GetItemLengthDisplayFrames() const
+{
+	if (const USequencerPlaylistItem_Sequence* SequenceItem = Cast<const USequencerPlaylistItem_Sequence>(GetItem()))
+	{
+		if (SequenceItem->Sequence && SequenceItem->Sequence->GetMovieScene())
+		{
+			UMovieScene* MovieScene = SequenceItem->Sequence->GetMovieScene();
+			TRange<FFrameNumber> Range = MovieScene->GetPlaybackRange();
+			if (Range.GetLowerBound().IsClosed() && Range.GetUpperBound().IsClosed())
+			{
+				return ConvertFrameTime(Range.Size<FFrameTime>(),
+					MovieScene->GetTickResolution(),
+					MovieScene->GetDisplayRate()).FloorToFrame().Value;
+			}
+		}
+	}
+
+	return TOptional<int32>();
+}
+
+
 TSharedRef<SWidget> SSequencerPlaylistItemWidget::EnsureSelectedAndBuildContextMenu()
 {
 	TArray<UObject*> SelectedItems;
diff --git a/Engine/Plugins/VirtualProduction/SequencerPlaylists/Source/SequencerPlaylists/Private/SequencerPlaylistsWidgets.h b/Engine/Plugins/VirtualProduction/SequencerPlaylists/Source/SequencerPlaylists/Private/SequencerPlaylistsWidgets.h
--- a/Engine/Plugins/VirtualProduction/SequencerPlaylists/Source/SequencerPlaylists/Private/SequencerPlaylistsWidgets.h
+++ b/Engine/Plugins/VirtualProduction/SequencerPlaylists/Source/SequencerPlaylists/Private/SequencerPlaylistsWidgets.h
@@ -157,7 +157,10 @@ public:
 	void Construct(const FArguments& InArgs, TSharedPtr<FSequencerPlaylistRowData> InRowData, const TSharedRef<STableViewBase>& OwnerTableView);
 
 	const TSharedPtr<FSequencerPlaylistRowData>& GetRowData() { return RowData; }
+	TSharedPtr<const FSequencerPlaylistRowData> GetRowData() const { return RowData; }
+
 	USequencerPlaylistItem* GetItem() { return GetRowData() ? GetRowData()->WeakItem.Get() : nullptr; }
+	const USequencerPlaylistItem* GetItem() const { return GetRowData() ? GetRowData()->WeakItem.Get() : nullptr; }
 
 	//~ Begin STableRow
 	void ConstructChildren(ETableViewMode::Type InOwnerTableMode, const TAttribute<FMargin>& InPadding, const TSharedRef<SWidget>& InContent) override;
@@ -179,6 +182,7 @@ private:
 
 	EVisibility GetRowDimmingVisibility() const;
 	EVisibility GetPlayModeTransportVisibility() const;
+	TOptional<int32> GetItemLengthDisplayFrames() const;
 
 	TSharedRef<SWidget> EnsureSelectedAndBuildContextMenu();
 	TSharedRef<SWidget> BuildContextMenu(const TArray<UObject*>& SelectedItems);
diff --git a/Engine/Plugins/VirtualProduction/SequencerPlaylists/Source/SequencerPlaylists/Public/SequencerPlaylistItem.h b/Engine/Plugins/VirtualProduction/SequencerPlaylists/Source/SequencerPlaylists/Public/SequencerPlaylistItem.h
--- a/Engine/Plugins/VirtualProduction/SequencerPlaylists/Source/SequencerPlaylists/Public/SequencerPlaylistItem.h
+++ b/Engine/Plugins/VirtualProduction/SequencerPlaylists/Source/SequencerPlaylists/Public/SequencerPlaylistItem.h
@@ -20,29 +20,29 @@ public:
 
 public:
 	/** Number of frames by which to clip the in point of sections played from this item. Will also affect the first frame for hold. */
-	UPROPERTY(EditAnywhere, BlueprintReadWrite, Category="SequencerPlaylists", meta=(NoResetToDefault))
+	UPROPERTY(EditAnywhere, BlueprintReadWrite, Category="SequencerPlaylists")
 	int32 StartFrameOffset;
 
 	/** Number of frames by which to clip the out point of sections played from this item. */
-	UPROPERTY(EditAnywhere, BlueprintReadWrite, Category="SequencerPlaylists", meta=(NoResetToDefault))
+	UPROPERTY(EditAnywhere, BlueprintReadWrite, Category="SequencerPlaylists")
 	int32 EndFrameOffset;
 
 	/**
 	 * If true, the sequence will be inserted immediately on recording start and any time Reset()
 	 * is called, paused at the first frame indefinitely until subsequently played or stopped.
 	 */
-	UPROPERTY(EditAnywhere, BlueprintReadWrite, Category="SequencerPlaylists", meta=(NoResetToDefault))
+	UPROPERTY(EditAnywhere, BlueprintReadWrite, Category="SequencerPlaylists")
 	bool bHoldAtFirstFrame;
 
 	/** 0 is single playthrough, >= 1 is (n+1) playthroughs. */
-	UPROPERTY(EditAnywhere, BlueprintReadWrite, Category="SequencerPlaylists", meta=(NoResetToDefault))
+	UPROPERTY(EditAnywhere, BlueprintReadWrite, Category="SequencerPlaylists")
 	int32 NumLoops;
 
 	/** Speed multiplier at which to play sections created by this item. A value of 1 will result in playback at the original speed. */
-	UPROPERTY(EditAnywhere, BlueprintReadWrite, Category="SequencerPlaylists", meta=(NoResetToDefault))
+	UPROPERTY(EditAnywhere, BlueprintReadWrite, Category="SequencerPlaylists")
 	float PlaybackSpeed = 1.0f;
 
-	/** Disables playback of this item and prevents it from entering a hold state. */
-	UPROPERTY(EditAnywhere, BlueprintReadWrite, Category="SequencerPlaylists", meta=(NoResetToDefault))
+	/** Disable playback of this item and prevent it from entering a hold state. */
+	UPROPERTY(EditAnywhere, BlueprintReadWrite, Category="SequencerPlaylists")
 	bool bMute;
 };
