SequencerPlaylists: avoid creating an empty transaction when stopping a take recording does not modify sequences

diff --git a/Engine/Plugins/VirtualProduction/SequencerPlaylists/Source/SequencerPlaylists/Private/SequencerPlaylistItem_Sequence.cpp b/Engine/Plugins/VirtualProduction/SequencerPlaylists/Source/SequencerPlaylists/Private/SequencerPlaylistItem_Sequence.cpp
--- a/Engine/Plugins/VirtualProduction/SequencerPlaylists/Source/SequencerPlaylists/Private/SequencerPlaylistItem_Sequence.cpp
+++ b/Engine/Plugins/VirtualProduction/SequencerPlaylists/Source/SequencerPlaylists/Private/SequencerPlaylistItem_Sequence.cpp
@@ -134,10 +134,12 @@ bool FSequencerPlaylistItemPlayer_Sequence::Play(USequencerPlaylistItem* Item)
 
 bool FSequencerPlaylistItemPlayer_Sequence::Stop(USequencerPlaylistItem* Item)
 {
+	bool bSequenceWasModified = false;
+
 	USequencerPlaylistItem_Sequence* SequenceItem = CastChecked<USequencerPlaylistItem_Sequence>(Item);
 	if (!SequenceItem || !SequenceItem->Sequence)
 	{
-		return false;
+		return bSequenceWasModified;
 	}
 
 	FItemState& ItemState = ItemStates.FindOrAdd(Item);
@@ -145,12 +147,12 @@ bool FSequencerPlaylistItemPlayer_Sequence::Stop(USequencerPlaylistItem* Item)
 
 	if (ItemState.WeakPlaySections.Num() == 0 && HoldSection == nullptr)
 	{
-		return false;
+		return bSequenceWasModified;
 	}
 
 	if (HoldSection)
 	{
-		EndSection(HoldSection);
+		bSequenceWasModified |= EndSection(HoldSection);
 		ItemState.WeakHoldSection.Reset();
 	}
 
@@ -158,7 +160,7 @@ bool FSequencerPlaylistItemPlayer_Sequence::Stop(USequencerPlaylistItem* Item)
 	{
 		if (UMovieSceneSubSection* PlaySection = WeakPlaySection.Get())
 		{
-			EndSection(PlaySection);
+			bSequenceWasModified |= EndSection(PlaySection);
 		}
 	}
 
@@ -170,7 +172,7 @@ bool FSequencerPlaylistItemPlayer_Sequence::Stop(USequencerPlaylistItem* Item)
 		ItemState.PlayingUntil_RootTicks = FMath::Min(Now.Value, ItemState.PlayingUntil_RootTicks);
 	}
 
-	return true;
+	return bSequenceWasModified;
 }
 
 
@@ -342,23 +344,23 @@ UMovieSceneSubTrack* FSequencerPlaylistItemPlayer_Sequence::GetOrCreateWorkingTr
 }
 
 
-void FSequencerPlaylistItemPlayer_Sequence::EndSection(UMovieSceneSection* Section)
+bool FSequencerPlaylistItemPlayer_Sequence::EndSection(UMovieSceneSection* Section)
 {
+	bool bSequenceWasModified = false;
+
 	if (!ensure(Section))
 	{
-		return;
+		return bSequenceWasModified;
 	}
 
 	TSharedPtr<ISequencer> Sequencer = WeakSequencer.Pin();
 	if (!ensure(Sequencer))
 	{
-		return;
+		return bSequenceWasModified;
 	}
 
 	ULevelSequence* RootSequence = Cast<ULevelSequence>(Sequencer->GetRootMovieSceneSequence());
 
-	RootSequence->Modify();
-
 	UMovieScene* SectionScene = Section->GetTypedOuter<UMovieScene>();
 	check(SectionScene);
 
@@ -366,6 +368,9 @@ void FSequencerPlaylistItemPlayer_Sequence::EndSection(UMovieSceneSection* Secti
 	const FFrameTime SectionNow = GlobalTime.ConvertTo(SectionScene->GetTickResolution());
 	if (Section->IsTimeWithinSection(SectionNow.FloorToFrame()))
 	{
+		RootSequence->Modify();
+		bSequenceWasModified = true;
+
 		Section->SetEndFrame(TRangeBound<FFrameNumber>::Exclusive(SectionNow.FloorToFrame()));
 
 		// Remove degenerate sections. This can happen if we reset then stop a held item while paused.
@@ -375,7 +380,7 @@ void FSequencerPlaylistItemPlayer_Sequence::EndSection(UMovieSceneSection* Secti
 			{
 				SectionTrack->Modify();
 				SectionTrack->RemoveSection(*Section);
-				return;
+				return bSequenceWasModified;
 			}
 		}
 
@@ -396,6 +401,8 @@ void FSequencerPlaylistItemPlayer_Sequence::EndSection(UMovieSceneSection* Secti
 			}
 		}
 	}
+
+	return bSequenceWasModified;
 }
 
 
diff --git a/Engine/Plugins/VirtualProduction/SequencerPlaylists/Source/SequencerPlaylists/Private/SequencerPlaylistItem_Sequence.h b/Engine/Plugins/VirtualProduction/SequencerPlaylists/Source/SequencerPlaylists/Private/SequencerPlaylistItem_Sequence.h
--- a/Engine/Plugins/VirtualProduction/SequencerPlaylists/Source/SequencerPlaylists/Private/SequencerPlaylistItem_Sequence.h
+++ b/Engine/Plugins/VirtualProduction/SequencerPlaylists/Source/SequencerPlaylists/Private/SequencerPlaylistItem_Sequence.h
@@ -54,7 +54,12 @@ private:
 	void ClearItemStates();
 
 	UMovieSceneSubTrack* GetOrCreateWorkingTrack(USequencerPlaylistItem* Item);
-	void EndSection(UMovieSceneSection* Section);
+
+	/**
+	 * Marks the end of a section based on the item player's Sequencer.
+	 * @return true if the current sequence was modified, otherwise false.
+	 */
+	bool EndSection(UMovieSceneSection* Section);
 
 private:
 	TMap<USequencerPlaylistItem*, FItemState> ItemStates;
diff --git a/Engine/Plugins/VirtualProduction/SequencerPlaylists/Source/SequencerPlaylists/Private/SequencerPlaylistPlayer.cpp b/Engine/Plugins/VirtualProduction/SequencerPlaylists/Source/SequencerPlaylists/Private/SequencerPlaylistPlayer.cpp
--- a/Engine/Plugins/VirtualProduction/SequencerPlaylists/Source/SequencerPlaylists/Private/SequencerPlaylistPlayer.cpp
+++ b/Engine/Plugins/VirtualProduction/SequencerPlaylists/Source/SequencerPlaylists/Private/SequencerPlaylistPlayer.cpp
@@ -473,10 +473,18 @@ void USequencerPlaylistPlayer::OnTakeRecorderStopped(UTakeRecorder* InRecorder)
 	if (TSharedPtr<ISequencer> Sequencer = GetValidatedSequencer())
 	{
 		FScopedTransaction Transaction(LOCTEXT("TakeRecorderStoppedTransaction", "Playlist - Take Recorder stopped"));
+		bool bAnySequencesModified = false;
 
 		for (USequencerPlaylistItem* Item : Playlist->Items)
 		{
-			GetCheckedItemPlayer(Item)->Stop(Item);
+			bAnySequencesModified |= GetCheckedItemPlayer(Item)->Stop(Item);
+		}
+
+		if (!bAnySequencesModified)
+		{
+			// Cancel the otherwise empty transaction if stopping the item
+			// players did not modify any sequences.
+			Transaction.Cancel();
 		}
 	}
 }
