Sequencer Playlists: Fix for crash with out of bounds offsets.

diff --git a/Engine/Plugins/VirtualProduction/SequencerPlaylists/Source/SequencerPlaylists/Private/SequencerPlaylistItem_Sequence.cpp b/Engine/Plugins/VirtualProduction/SequencerPlaylists/Source/SequencerPlaylists/Private/SequencerPlaylistItem_Sequence.cpp
--- a/Engine/Plugins/VirtualProduction/SequencerPlaylists/Source/SequencerPlaylists/Private/SequencerPlaylistItem_Sequence.cpp
+++ b/Engine/Plugins/VirtualProduction/SequencerPlaylists/Source/SequencerPlaylists/Private/SequencerPlaylistItem_Sequence.cpp
@@ -101,9 +101,11 @@ bool FSequencerPlaylistItemPlayer_Sequence::Play(USequencerPlaylistItem* Item)
 
 	const FQualifiedFrameTime GlobalTime = Sequencer->GetGlobalTime();
 	const FFrameNumber StartFrame = GlobalTime.Time.FloorToFrame();
+	const int32 MinDuration = 1;
 	const int32 MaxDuration = TNumericLimits<int32>::Max() - StartFrame.Value - 1;
-	const int32 Duration = FMath::Min(MaxDuration,
-		(SingleLoopDuration * FMath::Max(1, SequenceItem->NumLoops + 1)).FloorToFrame().Value);
+	const int32 Duration = FMath::Clamp(
+		(SingleLoopDuration * FMath::Max(1, SequenceItem->NumLoops + 1)).FloorToFrame().Value,
+		MinDuration, MaxDuration);
 
 	ItemState.PlayingUntil_RootTicks = FMath::Max(StartFrame.Value + Duration, ItemState.PlayingUntil_RootTicks);
 
diff --git a/Engine/Plugins/VirtualProduction/SequencerPlaylists/Source/SequencerPlaylists/Private/SequencerPlaylistsWidgets.cpp b/Engine/Plugins/VirtualProduction/SequencerPlaylists/Source/SequencerPlaylists/Private/SequencerPlaylistsWidgets.cpp
--- a/Engine/Plugins/VirtualProduction/SequencerPlaylists/Source/SequencerPlaylists/Private/SequencerPlaylistsWidgets.cpp
+++ b/Engine/Plugins/VirtualProduction/SequencerPlaylists/Source/SequencerPlaylists/Private/SequencerPlaylistsWidgets.cpp
@@ -1135,6 +1135,7 @@ TSharedRef<SWidget> SSequencerPlaylistItemWidget::GenerateWidgetForColumn(const
 				SNew(SNumericEntryBox<int32>)
 				.AllowSpin(true)
 				.ToolTipText(LOCTEXT("StartFrameOffsetTooltip", "Number of frames by which to clip the in point of sections played from this item. Will also affect the first frame for hold."))
+				.MinSliderValue(0)
 				.MaxSliderValue(this, &SSequencerPlaylistItemWidget::GetItemLengthDisplayFrames)
 				.Value(TAttribute<TOptional<int32>>::CreateLambda([WeakItem]() { return WeakItem.IsValid() ? WeakItem->StartFrameOffset : TOptional<int32>(); }))
 				.OnValueChanged_Lambda([WeakItem](int32 NewValue) {
@@ -1155,6 +1156,7 @@ TSharedRef<SWidget> SSequencerPlaylistItemWidget::GenerateWidgetForColumn(const
 				SNew(SNumericEntryBox<int32>)
 				.AllowSpin(true)
 				.ToolTipText(LOCTEXT("EndFrameOffsetTooltip", "Number of frames by which to clip the out point of sections played from this item."))
+				.MinSliderValue(0)
 				.MaxSliderValue(this, &SSequencerPlaylistItemWidget::GetItemLengthDisplayFrames)
 				.Value(TAttribute<TOptional<int32>>::CreateLambda([WeakItem]() { return WeakItem.IsValid() ? WeakItem->EndFrameOffset : TOptional<int32>(); }))
 				.OnValueChanged_Lambda([WeakItem](int32 NewValue) {
@@ -1417,7 +1419,7 @@ TOptional<int32> SSequencerPlaylistItemWidget::GetItemLengthDisplayFrames() cons
 		}
 	}
 
-	return TOptional<int32>();
+	return 0;
 }
 
 
