Sequencer Playlists: Support opening multiple Playlist assets simultaneously, and implement asset type actions. Also moves playlist item property grid to the bottom of the context menu.

diff --git a/Engine/Plugins/VirtualProduction/SequencerPlaylists/Source/SequencerPlaylists/Private/AssetTypeActions_SequencerPlaylist.cpp b/Engine/Plugins/VirtualProduction/SequencerPlaylists/Source/SequencerPlaylists/Private/AssetTypeActions_SequencerPlaylist.cpp
new file mode 100644
--- /dev/null
+++ b/Engine/Plugins/VirtualProduction/SequencerPlaylists/Source/SequencerPlaylists/Private/AssetTypeActions_SequencerPlaylist.cpp
@@ -0,0 +1,25 @@
+// Copyright Epic Games, Inc. All Rights Reserved.
+
+#include "AssetTypeActions_SequencerPlaylist.h"
+
+#include "SequencerPlaylist.h"
+#include "SequencerPlaylistsModule.h"
+
+
+
+UClass* FAssetTypeActions_SequencerPlaylist::GetSupportedClass() const
+{
+	return USequencerPlaylist::StaticClass();
+}
+
+void FAssetTypeActions_SequencerPlaylist::OpenAssetEditor(const TArray<UObject*>& InObjects, TSharedPtr<class IToolkitHost> InEditWithinLevelEditor /* = TSharedPtr<IToolkitHost>() */)
+{
+	for (UObject* PlaylistObj : InObjects)
+	{
+		if (USequencerPlaylist* Playlist = Cast<USequencerPlaylist>(PlaylistObj))
+		{
+			FSequencerPlaylistsModule& Module = static_cast<FSequencerPlaylistsModule&>(FSequencerPlaylistsModule::Get());
+			Module.SpawnEditorForPlaylist(Playlist);
+		}
+	}
+}
diff --git a/Engine/Plugins/VirtualProduction/SequencerPlaylists/Source/SequencerPlaylists/Private/AssetTypeActions_SequencerPlaylist.h b/Engine/Plugins/VirtualProduction/SequencerPlaylists/Source/SequencerPlaylists/Private/AssetTypeActions_SequencerPlaylist.h
new file mode 100644
--- /dev/null
+++ b/Engine/Plugins/VirtualProduction/SequencerPlaylists/Source/SequencerPlaylists/Private/AssetTypeActions_SequencerPlaylist.h
@@ -0,0 +1,17 @@
+// Copyright Epic Games, Inc. All Rights Reserved.
+
+#pragma once
+
+#include "CoreMinimal.h"
+#include "AssetTypeActions_Base.h"
+
+class FAssetTypeActions_SequencerPlaylist
+	: public FAssetTypeActions_Base
+{
+public:
+	virtual FText GetName() const override { return NSLOCTEXT("AssetTypeActions", "AssetTypeActions_SequencerPlaylist", "Sequencer Playlist"); }
+	virtual uint32 GetCategories() override { return EAssetTypeCategories::None; }
+	virtual FColor GetTypeColor() const override { return FColor(0, 232, 0); }
+	virtual UClass* GetSupportedClass() const override;
+	virtual void OpenAssetEditor(const TArray<UObject*>& InObjects, TSharedPtr<class IToolkitHost> InEditWithinLevelEditor /* = TSharedPtr<IToolkitHost>() */) override;
+};
diff --git a/Engine/Plugins/VirtualProduction/SequencerPlaylists/Source/SequencerPlaylists/Private/SequencerPlaylistItem_Sequence.h b/Engine/Plugins/VirtualProduction/SequencerPlaylists/Source/SequencerPlaylists/Private/SequencerPlaylistItem_Sequence.h
--- a/Engine/Plugins/VirtualProduction/SequencerPlaylists/Source/SequencerPlaylists/Private/SequencerPlaylistItem_Sequence.h
+++ b/Engine/Plugins/VirtualProduction/SequencerPlaylists/Source/SequencerPlaylists/Private/SequencerPlaylistItem_Sequence.h
@@ -10,6 +10,7 @@
 
 class ISequencer;
 class ULevelSequence;
+class UMovieSceneSection;
 class UMovieSceneSubSection;
 class UMovieSceneSubTrack;
 
diff --git a/Engine/Plugins/VirtualProduction/SequencerPlaylists/Source/SequencerPlaylists/Private/SequencerPlaylistPlayer.cpp b/Engine/Plugins/VirtualProduction/SequencerPlaylists/Source/SequencerPlaylists/Private/SequencerPlaylistPlayer.cpp
--- a/Engine/Plugins/VirtualProduction/SequencerPlaylists/Source/SequencerPlaylists/Private/SequencerPlaylistPlayer.cpp
+++ b/Engine/Plugins/VirtualProduction/SequencerPlaylists/Source/SequencerPlaylists/Private/SequencerPlaylistPlayer.cpp
@@ -50,10 +50,14 @@ void USequencerPlaylistPlayer::BeginDestroy()
 }
 
 
-//static
-USequencerPlaylistPlayer* USequencerPlaylistPlayer::GetDefaultPlayer()
+void USequencerPlaylistPlayer::SetPlaylist(USequencerPlaylist* InPlaylist)
 {
-	return ISequencerPlaylistsModule::Get().GetDefaultPlayer();
+	// Broadcast the event first, so subscribers can still access the previous via GetPlaylist().
+	if (InPlaylist != Playlist)
+	{
+		OnPlaylistSet.Broadcast(this, InPlaylist);
+		Playlist = InPlaylist;
+	}
 }
 
 
diff --git a/Engine/Plugins/VirtualProduction/SequencerPlaylists/Source/SequencerPlaylists/Private/SequencerPlaylistsModule.cpp b/Engine/Plugins/VirtualProduction/SequencerPlaylists/Source/SequencerPlaylists/Private/SequencerPlaylistsModule.cpp
--- a/Engine/Plugins/VirtualProduction/SequencerPlaylists/Source/SequencerPlaylists/Private/SequencerPlaylistsModule.cpp
+++ b/Engine/Plugins/VirtualProduction/SequencerPlaylists/Source/SequencerPlaylists/Private/SequencerPlaylistsModule.cpp
@@ -1,12 +1,14 @@
 // Copyright Epic Games, Inc. All Rights Reserved.
 
 #include "SequencerPlaylistsModule.h"
+#include "AssetTypeActions_SequencerPlaylist.h"
 #include "SequencerPlaylist.h"
 #include "SequencerPlaylistItem_Sequence.h"
 #include "SequencerPlaylistPlayer.h"
 #include "SequencerPlaylistsLog.h"
 #include "SequencerPlaylistsCommands.h"
 #include "SequencerPlaylistsStyle.h"
+#include "SequencerPlaylistsSubsystem.h"
 #include "SequencerPlaylistsWidgets.h"
 
 #include "ToolMenus.h"
@@ -15,7 +17,9 @@
 #include "WorkspaceMenuStructureModule.h"
 
 
-static const FName SequencerPlaylistsTabName("SequencerPlaylists");
+const FName FSequencerPlaylistsModule::MainTabType("SequencerPlaylists");
+const FName FSequencerPlaylistsModule::AssetEditorTabType("SequencerPlaylistsEditor");
+
 
 DEFINE_LOG_CATEGORY(LogSequencerPlaylists)
 
@@ -27,9 +31,6 @@ void FSequencerPlaylistsModule::StartupModule()
 {
 	// This code will execute after your module is loaded into memory; the exact timing is specified in the .uplugin file per-module
 
-	DefaultPlayer.Reset(NewObject<USequencerPlaylistPlayer>());
-	DefaultPlayer->SetPlaylist(NewObject<USequencerPlaylist>());
-
 	RegisterItemPlayer(USequencerPlaylistItem_Sequence::StaticClass(),
 		FSequencerPlaylistItemPlayerFactory::CreateLambda(
 			[](TSharedRef<ISequencer> Sequencer) {
@@ -51,11 +52,31 @@ void FSequencerPlaylistsModule::StartupModule()
 
 	UToolMenus::RegisterStartupCallback(FSimpleMulticastDelegate::FDelegate::CreateRaw(this, &FSequencerPlaylistsModule::RegisterMenus));
 
-	FGlobalTabmanager::Get()->RegisterNomadTabSpawner(SequencerPlaylistsTabName, FOnSpawnTab::CreateRaw(this, &FSequencerPlaylistsModule::OnSpawnPluginTab))
-		.SetDisplayName(LOCTEXT("TabSpawnerDisplayName", "Playlists"))
+	const FText TabSpawnerDisplayName = LOCTEXT("TabSpawnerDisplayName", "Playlists");
+	FGlobalTabmanager::Get()->RegisterNomadTabSpawner(MainTabType, FOnSpawnTab::CreateRaw(this, &FSequencerPlaylistsModule::OnSpawnPluginTab))
+		.SetDisplayName(TabSpawnerDisplayName)
 		.SetTooltipText(LOCTEXT("TabSpawnerTooltipText", "Open the Sequencer Playlists tab."))
 		.SetGroup(WorkspaceMenu::GetMenuStructure().GetLevelEditorVirtualProductionCategory())
 		.SetIcon(FSlateIcon(FSequencerPlaylistsStyle::GetStyleSetName(), "SequencerPlaylists.TabIcon"));
+
+	// Hidden spawner that permits multiple instances, each spawned without the SaveLayout flag.
+	FGlobalTabmanager::Get()->RegisterNomadTabSpawner(AssetEditorTabType, FOnSpawnTab::CreateRaw(this, &FSequencerPlaylistsModule::OnSpawnPluginTab))
+		.SetMenuType(ETabSpawnerMenuType::Hidden)
+		.SetDisplayName(TabSpawnerDisplayName)
+		.SetIcon(FSlateIcon(FSequencerPlaylistsStyle::GetStyleSetName(), "SequencerPlaylists.TabIcon"))
+		.SetReuseTabMethod(FOnFindTabToReuse::CreateLambda([] (const FTabId&)
+			{
+				return nullptr;
+			}
+		));
+
+	const FVector2D DefaultPlaylistSize(800.0, 600.0);
+	FTabManager::RegisterDefaultTabWindowSize(MainTabType, DefaultPlaylistSize);
+	FTabManager::RegisterDefaultTabWindowSize(AssetEditorTabType, DefaultPlaylistSize);
+
+	IAssetTools& AssetTools = FModuleManager::LoadModuleChecked<FAssetToolsModule>(TEXT("AssetTools")).Get();
+	AssetTypeActions = MakeShared<FAssetTypeActions_SequencerPlaylist>();
+	AssetTools.RegisterAssetTypeActions(AssetTypeActions.ToSharedRef());
 }
 
 
@@ -72,7 +93,17 @@ void FSequencerPlaylistsModule::ShutdownModule()
 
 	FSequencerPlaylistsCommands::Unregister();
 
-	FGlobalTabmanager::Get()->UnregisterNomadTabSpawner(SequencerPlaylistsTabName);
+	FGlobalTabmanager::Get()->UnregisterNomadTabSpawner(MainTabType);
+	FGlobalTabmanager::Get()->UnregisterNomadTabSpawner(AssetEditorTabType);
+
+	FTabManager::UnregisterDefaultTabWindowSize(MainTabType);
+	FTabManager::UnregisterDefaultTabWindowSize(AssetEditorTabType);
+
+	if (FModuleManager::Get().IsModuleLoaded(TEXT("AssetTools")))
+	{
+		IAssetTools& AssetTools = FModuleManager::LoadModuleChecked<FAssetToolsModule>("AssetTools").Get();
+		AssetTools.UnregisterAssetTypeActions(AssetTypeActions.ToSharedRef());
+	}
 }
 
 
@@ -101,17 +132,52 @@ TSharedPtr<ISequencerPlaylistItemPlayer> FSequencerPlaylistsModule::CreateItemPl
 
 TSharedRef<SDockTab> FSequencerPlaylistsModule::OnSpawnPluginTab(const FSpawnTabArgs& SpawnTabArgs)
 {
-	return SNew(SDockTab)
-		.TabRole(ETabRole::NomadTab)
+	TSharedRef<SDockTab> DockTab = SNew(SDockTab)
+		.TabRole(ETabRole::NomadTab);
+
+	TSharedRef<SSequencerPlaylistPanel> Panel = SNew(SSequencerPlaylistPanel, DockTab);
+
+	USequencerPlaylistsSubsystem* Subsystem = GEditor->GetEditorSubsystem<USequencerPlaylistsSubsystem>();
+	USequencerPlaylistPlayer* Player = Subsystem->CreatePlayerForEditor(Panel);
+	Panel->SetPlayer(Player);
+
+	DockTab->SetContent(Panel);
+	DockTab->SetLabel(TAttribute<FText>::CreateLambda(
 		[
-			SNew(SSequencerPlaylistPanel, DefaultPlayer.Get())
-		];
+			WeakPanel = TWeakPtr<SSequencerPlaylistPanel>(Panel)
+		]() -> FText
+		{
+			if (TSharedPtr<SSequencerPlaylistPanel> Panel = WeakPanel.Pin())
+			{
+				return FText::AsCultureInvariant(Panel->GetDisplayTitle());
+			}
+
+			return FText::GetEmpty();
+		}
+	));
+
+	return DockTab;
+}
+
+
+TSharedRef<SDockTab> FSequencerPlaylistsModule::SpawnEditorForPlaylist(USequencerPlaylist* Playlist)
+{
+	// ETabIdFlags::SaveLayout flag is unset for tabs other than the one spawned via the menu.
+	FTabId TabId = FTabId(AssetEditorTabType, ETabIdFlags::None);
+
+	TSharedPtr<SDockTab> PlaylistTab = FGlobalTabmanager::Get()->TryInvokeTab(TabId);
+	check(PlaylistTab);
+
+	TSharedRef<SSequencerPlaylistPanel> Panel = StaticCastSharedRef<SSequencerPlaylistPanel>(PlaylistTab->GetContent());
+	Panel->LoadPlaylist(Playlist);
+
+	return PlaylistTab.ToSharedRef();
 }
 
 
 void FSequencerPlaylistsModule::PluginButtonClicked()
 {
-	FGlobalTabmanager::Get()->TryInvokeTab(SequencerPlaylistsTabName);
+	FGlobalTabmanager::Get()->TryInvokeTab(MainTabType);
 }
 
 
diff --git a/Engine/Plugins/VirtualProduction/SequencerPlaylists/Source/SequencerPlaylists/Private/SequencerPlaylistsModule.h b/Engine/Plugins/VirtualProduction/SequencerPlaylists/Source/SequencerPlaylists/Private/SequencerPlaylistsModule.h
--- a/Engine/Plugins/VirtualProduction/SequencerPlaylists/Source/SequencerPlaylists/Private/SequencerPlaylistsModule.h
+++ b/Engine/Plugins/VirtualProduction/SequencerPlaylists/Source/SequencerPlaylists/Private/SequencerPlaylistsModule.h
@@ -9,6 +9,7 @@
 #include "SequencerPlaylistItem.h"
 
 
+class FAssetTypeActions_SequencerPlaylist;
 class FUICommandList;
 class SDockTab;
 class USequencerPlaylist;
@@ -17,6 +18,9 @@ class USequencerPlaylist;
 class FSequencerPlaylistsModule : public ISequencerPlaylistsModule
 {
 public:
+	static const FName MainTabType;
+	static const FName AssetEditorTabType;
+
 	//~ Begin IModuleInterface
 	void StartupModule() override;
 	void ShutdownModule() override;
@@ -24,18 +28,20 @@ public:
 
 	//~ Begin ISequencerPlaylistsModule
 	bool RegisterItemPlayer(TSubclassOf<USequencerPlaylistItem> ItemClass, FSequencerPlaylistItemPlayerFactory PlayerFactory) override;
-	USequencerPlaylistPlayer* GetDefaultPlayer() override { return DefaultPlayer.Get(); }
 	//~ End ISequencerPlaylistsModule
 
 	TSharedPtr<ISequencerPlaylistItemPlayer> CreateItemPlayerForClass(TSubclassOf<USequencerPlaylistItem> ItemClass, TSharedRef<ISequencer> Sequencer);
 
 	TSharedPtr<FUICommandList> GetCommandList() { return PluginCommands; }
 
+	TSharedRef<SDockTab> SpawnEditorForPlaylist(USequencerPlaylist* Playlist);
+
 private:
 	/** This function will be bound to Command (by default it will bring up plugin window) */
 	void PluginButtonClicked();
 
 	void RegisterMenus();
+
 	TSharedRef<SDockTab> OnSpawnPluginTab(const class FSpawnTabArgs& SpawnTabArgs);
 
 private:
@@ -44,4 +50,6 @@ private:
 	TStrongObjectPtr<USequencerPlaylistPlayer> DefaultPlayer;
 
 	TSharedPtr<FUICommandList> PluginCommands;
+
+	TSharedPtr<FAssetTypeActions_SequencerPlaylist> AssetTypeActions;
 };
diff --git a/Engine/Plugins/VirtualProduction/SequencerPlaylists/Source/SequencerPlaylists/Private/SequencerPlaylistsSubsystem.cpp b/Engine/Plugins/VirtualProduction/SequencerPlaylists/Source/SequencerPlaylists/Private/SequencerPlaylistsSubsystem.cpp
new file mode 100644
--- /dev/null
+++ b/Engine/Plugins/VirtualProduction/SequencerPlaylists/Source/SequencerPlaylists/Private/SequencerPlaylistsSubsystem.cpp
@@ -0,0 +1,73 @@
+// Copyright Epic Games, Inc. All Rights Reserved.
+
+#include "SequencerPlaylistsSubsystem.h"
+#include "SequencerPlaylist.h"
+#include "SequencerPlaylistPlayer.h"
+#include "SequencerPlaylistsLog.h"
+
+
+void USequencerPlaylistsSubsystem::Initialize(FSubsystemCollectionBase& Collection)
+{
+	Super::Initialize(Collection);
+}
+
+
+void USequencerPlaylistsSubsystem::Deinitialize()
+{
+	Super::Deinitialize();
+
+	ensureMsgf(EditorPlaylists.Num() == 0, TEXT("USequencerPlaylistsSubsystem::EditorPlaylists leak"));
+	ensureMsgf(EditorPlayers.Num() == 0, TEXT("USequencerPlaylistsSubsystem::EditorPlayers leak"));
+	ensureMsgf(EditorPackages.Num() == 0, TEXT("USequencerPlaylistsSubsystem::EditorPackages leak"));
+}
+
+
+USequencerPlaylistPlayer* USequencerPlaylistsSubsystem::CreatePlayerForEditor(TSharedRef<SSequencerPlaylistPanel> Editor)
+{
+	FSequencerPlaylistEditorHandle EditorHandle(&Editor.Get());
+
+	TObjectPtr<USequencerPlaylistPlayer>* ExistingPlayer = EditorPlayers.Find(EditorHandle);
+	if (!ensure(ExistingPlayer == nullptr))
+	{
+		return *ExistingPlayer;
+	}
+
+	ensure(EditorPackages.Find(EditorHandle) == nullptr);
+	ensure(EditorPlaylists.Find(EditorHandle) == nullptr);
+	ensure(EditorPlayers.Find(EditorHandle) == nullptr);
+
+	USequencerPlaylist* NewPlaylist = CreateTransientPlaylistForEditor(Editor);
+	EditorPlaylists.Add(EditorHandle, NewPlaylist);
+	EditorPackages.Add(EditorHandle, NewPlaylist->GetPackage());
+
+	USequencerPlaylistPlayer* NewPlayer = NewObject<USequencerPlaylistPlayer>();
+	EditorPlayers.Add(EditorHandle, NewPlayer);
+	NewPlayer->SetPlaylist(NewPlaylist);
+
+	return NewPlayer;
+}
+
+
+USequencerPlaylist* USequencerPlaylistsSubsystem::CreateTransientPlaylistForEditor(TSharedRef<SSequencerPlaylistPanel> Editor)
+{
+	FName PackageName = ::MakeUniqueObjectName(nullptr, UPackage::StaticClass(), TEXT("/Engine/Transient/SequencerPlaylist"));
+	UPackage* PlaylistPackage = NewObject<UPackage>(nullptr, PackageName, RF_Transient | RF_Transactional);
+	//PlaylistPackage->AddToRoot();
+
+	USequencerPlaylist* Playlist = NewObject<USequencerPlaylist>(PlaylistPackage, TEXT("UntitledPlaylist"));
+	Playlist->SetFlags(RF_Transactional);
+
+	return Playlist;
+}
+
+
+void USequencerPlaylistsSubsystem::NotifyEditorClosed(SSequencerPlaylistPanel* Editor)
+{
+	check(Editor);
+
+	FSequencerPlaylistEditorHandle EditorHandle(Editor);
+
+	ensure(EditorPlaylists.Remove(EditorHandle));
+	ensure(EditorPlayers.Remove(EditorHandle));
+	ensure(EditorPackages.Remove(EditorHandle));
+}
diff --git a/Engine/Plugins/VirtualProduction/SequencerPlaylists/Source/SequencerPlaylists/Private/SequencerPlaylistsSubsystem.h b/Engine/Plugins/VirtualProduction/SequencerPlaylists/Source/SequencerPlaylists/Private/SequencerPlaylistsSubsystem.h
new file mode 100644
--- /dev/null
+++ b/Engine/Plugins/VirtualProduction/SequencerPlaylists/Source/SequencerPlaylists/Private/SequencerPlaylistsSubsystem.h
@@ -0,0 +1,67 @@
+// Copyright Epic Games, Inc. All Rights Reserved.
+
+#pragma once
+
+#include "EditorSubsystem.h"
+#include "SequencerPlaylistsSubsystem.generated.h"
+
+
+class SSequencerPlaylistPanel;
+class USequencerPlaylist;
+class USequencerPlaylistPlayer;
+
+
+USTRUCT()
+struct FSequencerPlaylistEditorHandle
+{
+	GENERATED_BODY()
+
+	FSequencerPlaylistEditorHandle() {}
+
+	explicit FSequencerPlaylistEditorHandle(SSequencerPlaylistPanel* InEditor)
+		: Editor(InEditor)
+	{
+	}
+
+	bool operator==(const FSequencerPlaylistEditorHandle& Other) const
+	{
+		return Editor == Other.Editor;
+	}
+
+	friend uint32 GetTypeHash(const FSequencerPlaylistEditorHandle& Handle)
+	{
+		return GetTypeHash(Handle.Editor);
+	}
+
+	SSequencerPlaylistPanel* Editor = nullptr;
+};
+
+
+UCLASS()
+class USequencerPlaylistsSubsystem : public UEditorSubsystem
+{
+	GENERATED_BODY()
+
+public:
+	//~ Begin USubsystem interface
+	virtual void Initialize(FSubsystemCollectionBase& Collection) override;
+	virtual void Deinitialize() override;
+	//~ End USubsystem interface
+
+	/** For every tab, we manage the lifetime of the associated Player and its transient Playlist. */
+	USequencerPlaylistPlayer* CreatePlayerForEditor(TSharedRef<SSequencerPlaylistPanel> Editor);
+	void NotifyEditorClosed(SSequencerPlaylistPanel* Editor);
+
+protected:
+	USequencerPlaylist* CreateTransientPlaylistForEditor(TSharedRef<SSequencerPlaylistPanel> Editor);
+
+protected:
+	UPROPERTY(Transient)
+	TMap<FSequencerPlaylistEditorHandle, TObjectPtr<UPackage>> EditorPackages;
+
+	UPROPERTY(Transient)
+	TMap<FSequencerPlaylistEditorHandle, TObjectPtr<USequencerPlaylistPlayer>> EditorPlayers;
+
+	UPROPERTY(Transient)
+	TMap<FSequencerPlaylistEditorHandle, TObjectPtr<USequencerPlaylist>> EditorPlaylists;
+};
diff --git a/Engine/Plugins/VirtualProduction/SequencerPlaylists/Source/SequencerPlaylists/Private/SequencerPlaylistsWidgets.cpp b/Engine/Plugins/VirtualProduction/SequencerPlaylists/Source/SequencerPlaylists/Private/SequencerPlaylistsWidgets.cpp
--- a/Engine/Plugins/VirtualProduction/SequencerPlaylists/Source/SequencerPlaylists/Private/SequencerPlaylistsWidgets.cpp
+++ b/Engine/Plugins/VirtualProduction/SequencerPlaylists/Source/SequencerPlaylists/Private/SequencerPlaylistsWidgets.cpp
@@ -7,6 +7,7 @@
 #include "SequencerPlaylistPlayer.h"
 #include "SequencerPlaylistsModule.h"
 #include "SequencerPlaylistsStyle.h"
+#include "SequencerPlaylistsSubsystem.h"
 
 #include "AssetRegistry/AssetData.h"
 #include "AssetRegistry/AssetRegistryModule.h"
@@ -21,6 +22,7 @@
 #include "LevelSequence.h"
 #include "Misc/FileHelper.h"
 #include "Misc/TextFilter.h"
+#include "Misc/TransactionObjectEvent.h"
 #include "MovieScene.h"
 #include "ScopedTransaction.h"
 #include "SlateOptMacros.h"
@@ -54,10 +56,9 @@ const FName SSequencerPlaylistPanel::ColumnName_Loop(TEXT("Loop"));
 const FName SSequencerPlaylistPanel::ColumnName_HoverDetails(TEXT("HoverDetails"));
 
 
-void SSequencerPlaylistPanel::Construct(const FArguments& InArgs, USequencerPlaylistPlayer* InPlayer)
+void SSequencerPlaylistPanel::Construct(const FArguments& InArgs, TSharedPtr<SDockTab> InContainingTab)
 {
-	check(InPlayer);
-	WeakPlayer = InPlayer;
+	WeakContainingTab = InContainingTab;
 
 	SearchTextFilter = MakeShared<TTextFilter<const FSequencerPlaylistRowData&>>(
 		TTextFilter<const FSequencerPlaylistRowData&>::FItemToStringArray::CreateSP(this, &SSequencerPlaylistPanel::GetSearchStrings));
@@ -105,15 +106,7 @@ void SSequencerPlaylistPanel::Construct(const FArguments& InArgs, USequencerPlay
 				[
 					SNew(STextBlock)
 					.Font(FSequencerPlaylistsStyle::Get().GetFontStyle("SequencerPlaylists.TitleFont"))
-					.Text(TAttribute<FText>::CreateLambda([this]()
-						{
-							if (USequencerPlaylist* Playlist = GetCheckedPlaylist())
-							{
-								return FText::FromString(Playlist->GetName());
-							}
-
-							return FText::GetEmpty();
-						}))
+					.Text(TAttribute<FText>::CreateLambda([this]() { return FText::AsCultureInvariant(GetDisplayTitle()); }))
 				]
 				+ SHorizontalBox::Slot()
 				.AutoWidth()
@@ -142,7 +135,7 @@ void SSequencerPlaylistPanel::Construct(const FArguments& InArgs, USequencerPlay
 						.Font(FSequencerPlaylistsStyle::Get().GetFontStyle("SequencerPlaylists.DescriptionFont"))
 						.HintText(LOCTEXT("PlaylistDescriptionHint", "<playlist description>"))
 						.Text(TAttribute<FText>::CreateLambda([this]() {
-							if (USequencerPlaylist* Playlist = GetCheckedPlaylist())
+							if (USequencerPlaylist* Playlist = GetPlaylist())
 							{
 								return Playlist->Description;
 							}
@@ -150,7 +143,7 @@ void SSequencerPlaylistPanel::Construct(const FArguments& InArgs, USequencerPlay
 							return FText::GetEmpty();
 						}))
 						.OnTextCommitted_Lambda([this](const FText& NewText, ETextCommit::Type CommitType) {
-							if (USequencerPlaylist* Playlist = GetCheckedPlaylist())
+							if (USequencerPlaylist* Playlist = GetPlaylist())
 							{
 								Playlist->Description = NewText;
 							}
@@ -171,11 +164,80 @@ void SSequencerPlaylistPanel::Construct(const FArguments& InArgs, USequencerPlay
 			]
 		]
 	];
+}
+
+
+SSequencerPlaylistPanel::~SSequencerPlaylistPanel()
+{
+	GEditor->GetEditorSubsystem<USequencerPlaylistsSubsystem>()->NotifyEditorClosed(this);
+}
+
+
+FString SSequencerPlaylistPanel::GetDisplayTitle()
+{
+	FString DisplayTitle;
+
+	USequencerPlaylist* WorkingPlaylist = GetPlaylist();
+
+	if (USequencerPlaylist* LoadedFrom = WeakLoadedPlaylist.Get())
+	{
+		DisplayTitle = LoadedFrom->GetName();
+		if (WorkingPlaylist && WorkingPlaylist->GetPackage() && WorkingPlaylist->GetPackage()->IsDirty())
+		{
+			DisplayTitle.AppendChar(TEXT('*'));
+		}
+	}
+	else if (WorkingPlaylist)
+	{
+		DisplayTitle = WorkingPlaylist->GetName();
+	}
+
+	return DisplayTitle;
+}
+
+
+void SSequencerPlaylistPanel::SetPlayer(USequencerPlaylistPlayer* Player)
+{
+	WeakPlayer = Player;
 
 	RegenerateRows();
 }
 
 
+bool SSequencerPlaylistPanel::MatchesContext(const FTransactionContext& InContext, const TArray<TPair<UObject*, FTransactionObjectEvent>>& TransactionObjectContexts) const
+{
+	for (const TPair<UObject*, FTransactionObjectEvent>& TransactionObjectPair : TransactionObjectContexts)
+	{
+		UObject* Object = TransactionObjectPair.Key;
+		while (Object != nullptr)
+		{
+			if (Object->GetClass()->IsChildOf(USequencerPlaylist::StaticClass()))
+			{
+				return true;
+			}
+			Object = Object->GetOuter();
+		}
+	}
+
+	return false;
+}
+
+
+void SSequencerPlaylistPanel::PostUndo(bool bSuccess)
+{
+	if (bSuccess)
+	{
+		RegenerateRows();
+	}
+}
+
+
+void SSequencerPlaylistPanel::PostRedo(bool bSuccess)
+{
+	PostUndo(bSuccess);
+}
+
+
 TSharedRef<SWidget> SSequencerPlaylistPanel::Construct_LeftToolbar()
 {
 	FSequencerPlaylistsModule& Module = static_cast<FSequencerPlaylistsModule&>(ISequencerPlaylistsModule::Get());
@@ -406,6 +468,18 @@ USequencerPlaylist* SSequencerPlaylistPanel::GetCheckedPlaylist()
 }
 
 
+USequencerPlaylist* SSequencerPlaylistPanel::GetPlaylist()
+{
+	USequencerPlaylistPlayer* Player = WeakPlayer.Get();
+	if (!Player)
+	{
+		return nullptr;
+	}
+
+	return Player->GetPlaylist();
+}
+
+
 void SSequencerPlaylistPanel::RegenerateRows()
 {
 	USequencerPlaylist* Playlist = GetCheckedPlaylist();
@@ -563,7 +637,7 @@ void SSequencerPlaylistPanel::OnSavePlaylist()
 	}
 
 	FString PackageName = WeakLoadedPlaylist->GetPackage()->GetLoadedPath().GetPackageName();
-	SavePlaylist(PackageName);
+	SavePlaylistAs(PackageName);
 }
 
 
@@ -575,10 +649,10 @@ void SSequencerPlaylistPanel::OnSavePlaylistAs()
 		return;
 	}
 
-	SavePlaylist(PackageName);
+	SavePlaylistAs(PackageName);
 }
 
-void SSequencerPlaylistPanel::SavePlaylist(const FString& PackageName)
+void SSequencerPlaylistPanel::SavePlaylistAs(const FString& PackageName)
 {
 	USequencerPlaylist* Playlist = GetCheckedPlaylist();
 	if (!Playlist)
@@ -586,17 +660,11 @@ void SSequencerPlaylistPanel::SavePlaylist(const FString& PackageName)
 		return;
 	}
 
-	FScopedTransaction Transaction(LOCTEXT("SavePlaylistTransaction", "Save Sequencer Playlist"));
-
 	const FString AssetName = FPackageName::GetLongPackageAssetName(PackageName);
 	UPackage* Package = CreatePackage(*PackageName);
 	USequencerPlaylist* NewPlaylist = NewObject<USequencerPlaylist>(Package, *AssetName, RF_Public | RF_Standalone | RF_Transactional);
-
 	if (NewPlaylist)
 	{
-		// Within the transient package; no collision.
-		Playlist->Rename(*NewPlaylist->GetName());
-
 		NewPlaylist->Description = Playlist->Description;
 
 		for (USequencerPlaylistItem* Item : Playlist->Items)
@@ -611,38 +679,47 @@ void SSequencerPlaylistPanel::SavePlaylist(const FString& PackageName)
 	}
 
 	WeakLoadedPlaylist = NewPlaylist;
+	if (Playlist->GetPackage())
+	{
+		Playlist->GetPackage()->ClearDirtyFlag();
+	}
 }
 
 
 void SSequencerPlaylistPanel::OnLoadPlaylist(const FAssetData& InPreset)
 {
 	FSlateApplication::Get().DismissAllMenus();
+	LoadPlaylist(CastChecked<USequencerPlaylist>(InPreset.GetAsset()));
+}
 
-	USequencerPlaylist* Playlist = GetCheckedPlaylist();
-	if (!Playlist)
+
+void SSequencerPlaylistPanel::LoadPlaylist(USequencerPlaylist* PlaylistToLoad)
+{
+	if (!ensure(PlaylistToLoad))
 	{
 		return;
 	}
 
-	USequencerPlaylist* ImportedPlaylist = CastChecked<USequencerPlaylist>(InPreset.GetAsset());
-	if (ImportedPlaylist)
+	USequencerPlaylist* Playlist = GetCheckedPlaylist();
+	if (!Playlist)
 	{
-		FScopedTransaction Transaction(LOCTEXT("LoadPlaylistTransaction", "Load Sequencer Playlist"));
-
-		// Within the transient package; no collision.
-		Playlist->Rename(*ImportedPlaylist->GetName());
+		return;
+	}
 
-		Playlist->Description = ImportedPlaylist->Description;
+	Playlist->Description = PlaylistToLoad->Description;
 
-		Playlist->Items.Empty();
-		for (USequencerPlaylistItem* Item : ImportedPlaylist->Items)
-		{
-			Playlist->Items.Add(DuplicateObject<USequencerPlaylistItem>(Item, Playlist));
-		}
+	Playlist->Items.Empty();
+	for (USequencerPlaylistItem* Item : PlaylistToLoad->Items)
+	{
+		Playlist->Items.Add(DuplicateObject<USequencerPlaylistItem>(Item, Playlist));
+	}
 
-		RegenerateRows();
+	RegenerateRows();
 
-		WeakLoadedPlaylist = ImportedPlaylist;
+	WeakLoadedPlaylist = PlaylistToLoad;
+	if (Playlist->GetPackage())
+	{
+		Playlist->GetPackage()->ClearDirtyFlag();
 	}
 }
 
@@ -655,7 +732,6 @@ void SSequencerPlaylistPanel::OnNewPlaylist()
 		return;
 	}
 
-	Playlist->Rename(*MakeUniqueObjectName(GetTransientPackage(), USequencerPlaylist::StaticClass(), "SequencerPlaylist").ToString());
 	Playlist->Description = FText::GetEmpty();
 	Playlist->Items.Empty();
 	RegenerateRows();
@@ -727,7 +803,12 @@ FReply SSequencerPlaylistPanel::HandleClicked_AddSequence()
 		return FReply::Unhandled();
 	}
 
-	Playlist->Items.Add(NewObject<USequencerPlaylistItem_Sequence>(Playlist));
+	FScopedTransaction Transaction(LOCTEXT("AddSequenceTransaction", "Add Sequence Item To Playlist"));
+	Playlist->Modify();
+
+	USequencerPlaylistItem_Sequence* NewItem = NewObject<USequencerPlaylistItem_Sequence>(Playlist);
+	NewItem->SetFlags(RF_Transactional);
+	Playlist->Items.Add(NewItem);
 	RegenerateRows();
 
 	return FReply::Handled();
@@ -781,6 +862,9 @@ FReply SSequencerPlaylistPanel::HandleClicked_Item_Remove(SSequencerPlaylistItem
 		return FReply::Unhandled();
 	}
 
+	FScopedTransaction Transaction(LOCTEXT("RemoveSequenceTransaction", "Remove Sequence Item From Playlist"));
+	Playlist->Modify();
+
 	ensure(Playlist->Items.RemoveSingle(ItemWidget.GetItem()));
 	RegenerateRows();
 	return FReply::Handled();
@@ -1484,6 +1568,16 @@ TSharedRef<SWidget> SSequencerPlaylistItemWidget::BuildContextMenu(const TArray<
 		);
 	}
 
+	MenuBuilder.BeginSection(NAME_None, LOCTEXT("ItemContextEditHeading", "Edit"));
+	{
+		MenuBuilder.AddMenuEntry(
+			LOCTEXT("ItemContextRemove", "Remove from Playlist"),
+			LOCTEXT("ItemContextRemoveTooltip", "Remove this item from the Playlist"),
+			FSlateIcon(FAppStyle::GetAppStyleSetName(), "ContentBrowser.AssetActions.Delete"),
+			FUIAction(FExecuteAction::CreateLambda([this]() { RemoveClickedDelegate.Execute(*this); }))
+		);
+	}
+
 	MenuBuilder.BeginSection(NAME_None, LOCTEXT("ItemContextDetailsHeading", "Details"));
 	{
 		FDetailsViewArgs DetailsViewArgs;
@@ -1506,16 +1600,6 @@ TSharedRef<SWidget> SSequencerPlaylistItemWidget::BuildContextMenu(const TArray<
 		MenuBuilder.AddWidget(DetailsView, FText::GetEmpty(), true);
 	}
 
-	MenuBuilder.BeginSection(NAME_None, LOCTEXT("ItemContextEditHeading", "Edit"));
-	{
-		MenuBuilder.AddMenuEntry(
-			LOCTEXT("ItemContextRemove", "Remove from Playlist"),
-			LOCTEXT("ItemContextRemoveTooltip", "Remove this item from the Playlist"),
-			FSlateIcon(FAppStyle::GetAppStyleSetName(), "ContentBrowser.AssetActions.Delete"),
-			FUIAction(FExecuteAction::CreateLambda([this]() { RemoveClickedDelegate.Execute(*this); }))
-		);
-	}
-
 	return MenuBuilder.MakeWidget();
 }
 
diff --git a/Engine/Plugins/VirtualProduction/SequencerPlaylists/Source/SequencerPlaylists/Private/SequencerPlaylistsWidgets.h b/Engine/Plugins/VirtualProduction/SequencerPlaylists/Source/SequencerPlaylists/Private/SequencerPlaylistsWidgets.h
--- a/Engine/Plugins/VirtualProduction/SequencerPlaylists/Source/SequencerPlaylists/Private/SequencerPlaylistsWidgets.h
+++ b/Engine/Plugins/VirtualProduction/SequencerPlaylists/Source/SequencerPlaylists/Private/SequencerPlaylistsWidgets.h
@@ -3,6 +3,7 @@
 #pragma once
 
 #include "CoreMinimal.h"
+#include "EditorUndoClient.h"
 #include "PropertyEditorDelegates.h"
 #include "Widgets/DeclarativeSyntaxSupport.h"
 #include "Widgets/SCompoundWidget.h"
@@ -10,6 +11,7 @@
 
 
 struct FAssetData;
+class SDockTab;
 class SMenuAnchor;
 class SMultiLineEditableTextBox;
 class SSequencerPlaylistItemWidget;
@@ -34,7 +36,9 @@ struct FSequencerPlaylistRowData
 };
 
 
-class SSequencerPlaylistPanel : public SCompoundWidget
+class SSequencerPlaylistPanel
+	: public SCompoundWidget
+	, public FSelfRegisteringEditorUndoClient
 {
 	SLATE_BEGIN_ARGS(SSequencerPlaylistPanel) {}
 	SLATE_END_ARGS()
@@ -48,10 +52,21 @@ public:
 	static const FName ColumnName_Loop;
 	static const FName ColumnName_HoverDetails;
 
-	void Construct(const FArguments& InArgs, USequencerPlaylistPlayer* InPlayer);
+	void Construct(const FArguments& InArgs, TSharedPtr<SDockTab> InContainingTab);
+	virtual ~SSequencerPlaylistPanel();
 
 	bool InPlayMode() const { return bPlayMode; }
 
+	FString GetDisplayTitle();
+	void SetPlayer(USequencerPlaylistPlayer* Player);
+	void LoadPlaylist(USequencerPlaylist* Playlist);
+
+	//~ Begin FEditorUndoClient interface
+	virtual bool MatchesContext(const FTransactionContext& InContext, const TArray<TPair<UObject*, FTransactionObjectEvent>>& TransactionObjectContexts) const override;
+	virtual void PostUndo(bool bSuccess) override;
+	virtual void PostRedo(bool bSuccess) override;
+	//~ End FEditorUndoClient interface
+
 private:
 	TSharedRef<SWidget> Construct_LeftToolbar();
 	TSharedRef<SWidget> Construct_RightToolbar();
@@ -60,12 +75,13 @@ private:
 	TSharedRef<SWidget> Construct_ItemListView();
 
 	USequencerPlaylist* GetCheckedPlaylist();
+	USequencerPlaylist* GetPlaylist();
 	void RegenerateRows();
 
 	TSharedRef<SWidget> BuildOpenPlaylistMenu();
 	void OnSavePlaylist();
 	void OnSavePlaylistAs();
-	void SavePlaylist(const FString& PackageName);
+	void SavePlaylistAs(const FString& PackageName);
 	void OnLoadPlaylist(const FAssetData& InPreset);
 	void OnNewPlaylist();
 
@@ -94,6 +110,7 @@ private:
 
 	bool bPlayMode = false;
 
+	TWeakPtr<SDockTab> WeakContainingTab;
 	TSharedPtr<SSearchBox> SearchBox;
 	TSharedPtr<TTextFilter<const FSequencerPlaylistRowData&>> SearchTextFilter;
 	TArray<TSharedPtr<FSequencerPlaylistRowData>> ItemRows;
diff --git a/Engine/Plugins/VirtualProduction/SequencerPlaylists/Source/SequencerPlaylists/Public/ISequencerPlaylistsModule.h b/Engine/Plugins/VirtualProduction/SequencerPlaylists/Source/SequencerPlaylists/Public/ISequencerPlaylistsModule.h
--- a/Engine/Plugins/VirtualProduction/SequencerPlaylists/Source/SequencerPlaylists/Public/ISequencerPlaylistsModule.h
+++ b/Engine/Plugins/VirtualProduction/SequencerPlaylists/Source/SequencerPlaylists/Public/ISequencerPlaylistsModule.h
@@ -69,5 +69,6 @@ public:
 
 	virtual bool RegisterItemPlayer(TSubclassOf<USequencerPlaylistItem> ItemClass, FSequencerPlaylistItemPlayerFactory PlayerFactory) = 0;
 
-	virtual USequencerPlaylistPlayer* GetDefaultPlayer() = 0;
+	UE_DEPRECATED(5.1, "There is no longer a \"default\" player. Open a specific Playlist asset to create a player associated with it.")
+	virtual USequencerPlaylistPlayer* GetDefaultPlayer() { return nullptr; }
 };
diff --git a/Engine/Plugins/VirtualProduction/SequencerPlaylists/Source/SequencerPlaylists/Public/SequencerPlaylistItem.h b/Engine/Plugins/VirtualProduction/SequencerPlaylists/Source/SequencerPlaylists/Public/SequencerPlaylistItem.h
--- a/Engine/Plugins/VirtualProduction/SequencerPlaylists/Source/SequencerPlaylists/Public/SequencerPlaylistItem.h
+++ b/Engine/Plugins/VirtualProduction/SequencerPlaylists/Source/SequencerPlaylists/Public/SequencerPlaylistItem.h
@@ -6,10 +6,6 @@
 #include "SequencerPlaylistItem.generated.h"
 
 
-class ISequencer;
-class UMovieSceneSection;
-
-
 UCLASS(BlueprintType, Abstract, Within=SequencerPlaylist)
 class USequencerPlaylistItem : public UObject
 {
diff --git a/Engine/Plugins/VirtualProduction/SequencerPlaylists/Source/SequencerPlaylists/Public/SequencerPlaylistPlayer.h b/Engine/Plugins/VirtualProduction/SequencerPlaylists/Source/SequencerPlaylists/Public/SequencerPlaylistPlayer.h
--- a/Engine/Plugins/VirtualProduction/SequencerPlaylists/Source/SequencerPlaylists/Public/SequencerPlaylistPlayer.h
+++ b/Engine/Plugins/VirtualProduction/SequencerPlaylists/Source/SequencerPlaylists/Public/SequencerPlaylistPlayer.h
@@ -3,10 +3,9 @@
 #pragma once
 
 #include "CoreMinimal.h"
-#include "Templates/SubclassOf.h"
 #include "SequencerPlaylistItem.h"
+#include "Templates/SubclassOf.h"
 #include "Templates/UniquePtr.h"
-
 #include "Tickable.h"
 
 #include "SequencerPlaylistPlayer.generated.h"
@@ -15,10 +14,14 @@
 class ISequencer;
 class ISequencerPlaylistItemPlayer;
 class USequencerPlaylist;
+class USequencerPlaylistPlayer;
 class UTakeRecorder;
 class FTickablePlaylist;
 
 
+DECLARE_DYNAMIC_MULTICAST_DELEGATE_TwoParams(FOnPlayerSequencerPlaylistSet, USequencerPlaylistPlayer*, Player, USequencerPlaylist*, Playlist);
+
+
 /** Controls playback of playlist items */
 UCLASS(BlueprintType)
 class USequencerPlaylistPlayer : public UObject
@@ -33,13 +36,17 @@ public:
 	//~ End UObject interface
 
 	UFUNCTION(BlueprintCallable, Category="SequencerPlaylists")
-	void SetPlaylist(USequencerPlaylist* InPlaylist) { Playlist = InPlaylist; }
+	void SetPlaylist(USequencerPlaylist* InPlaylist);
 
 	UFUNCTION(BlueprintPure, Category="SequencerPlaylists")
 	USequencerPlaylist* GetPlaylist() { return Playlist; }
 
-	UFUNCTION(BlueprintPure, meta=(DisplayName="Get Default Sequencer Playlist Player"), Category="SequencerPlaylists")
-	static USequencerPlaylistPlayer* GetDefaultPlayer();
+	UE_DEPRECATED(5.1, "There is no longer a \"default\" player. Open a specific Playlist asset to create a player associated with it.")
+	UFUNCTION(BlueprintPure, meta=(DeprecatedFunction, DeprecationMessage= "There is no longer a \"default\" player. Open a specific Playlist asset to create a player associated with it.", DisplayName="Get Default Sequencer Playlist Player"), Category="SequencerPlaylists")
+	static USequencerPlaylistPlayer* GetDefaultPlayer() { return nullptr; }
+
+	UPROPERTY(BlueprintAssignable, Category="SequencerPlaylists")
+	FOnPlayerSequencerPlaylistSet OnPlaylistSet;
 
 public:
 	UFUNCTION(BlueprintCallable, Category="SequencerPlaylists")
@@ -72,7 +79,7 @@ private:
 	void EnterUnboundedPlayIfNotRecording();
 
 	/**
-	 * Ticked by a tickable game object to performe any necessary time-sliced logic
+	 * Ticked by a tickable game object to perform any necessary time-sliced logic
 	 */
 	void Tick(float DeltaTime);
 
@@ -115,7 +122,7 @@ public:
 		RETURN_QUICK_DECLARE_CYCLE_STAT(FTickablePlaylist, STATGROUP_Tickables);
 	}
 
-	//Make sure it always ticks, otherwise we can miss recording, in particularly when time code is always increasing throughout the system.
+	// Make sure it always ticks, otherwise we can miss recording, particularly when time code is always increasing throughout the system.
 	virtual ETickableTickType GetTickableTickType() const override { return ETickableTickType::Always; }
 
 	virtual bool IsTickableInEditor() const override
@@ -125,15 +132,15 @@ public:
 
 	virtual UWorld* GetTickableGameObjectWorld() const override
 	{
-		USequencerPlaylistPlayer* Playlist = WeakPlaylistPlayer.Get();
-		return Playlist ? Playlist->GetWorld() : nullptr;
+		USequencerPlaylistPlayer* Player = WeakPlaylistPlayer.Get();
+		return Player ? Player->GetWorld() : nullptr;
 	}
 
 	virtual void Tick(float DeltaTime) override
 	{
-		if (USequencerPlaylistPlayer* Playlist = WeakPlaylistPlayer.Get())
+		if (USequencerPlaylistPlayer* Player = WeakPlaylistPlayer.Get())
 		{
-			Playlist->Tick(DeltaTime);
+			Player->Tick(DeltaTime);
 		}
 	}
 
diff --git a/Engine/Plugins/VirtualProduction/SequencerPlaylists/Source/SequencerPlaylists/SequencerPlaylists.Build.cs b/Engine/Plugins/VirtualProduction/SequencerPlaylists/Source/SequencerPlaylists/SequencerPlaylists.Build.cs
--- a/Engine/Plugins/VirtualProduction/SequencerPlaylists/Source/SequencerPlaylists/SequencerPlaylists.Build.cs
+++ b/Engine/Plugins/VirtualProduction/SequencerPlaylists/Source/SequencerPlaylists/SequencerPlaylists.Build.cs
@@ -30,6 +30,7 @@ public class SequencerPlaylists : ModuleRules
 			"TakesCore",
 			"ToolWidgets",
 			"MovieSceneTools",
+			"EditorSubsystem",
 		});
 
 		DynamicallyLoadedModuleNames.AddRange(new string[] {
