Sequencer Playlists: Workaround for hold sections sometimes evaluating incorrectly with certain sequences.

Also removes some lingering unused code.

diff --git a/Engine/Plugins/VirtualProduction/SequencerPlaylists/Source/SequencerPlaylists/Private/SequencerPlaylistItem_Sequence.cpp b/Engine/Plugins/VirtualProduction/SequencerPlaylists/Source/SequencerPlaylists/Private/SequencerPlaylistItem_Sequence.cpp
--- a/Engine/Plugins/VirtualProduction/SequencerPlaylists/Source/SequencerPlaylists/Private/SequencerPlaylistItem_Sequence.cpp
+++ b/Engine/Plugins/VirtualProduction/SequencerPlaylists/Source/SequencerPlaylists/Private/SequencerPlaylistItem_Sequence.cpp
@@ -234,6 +234,10 @@ bool FSequencerPlaylistItemPlayer_Sequence::AddHold(USequencerPlaylistItem* Item
 
 	HoldSection->Parameters.TimeScale = 0.f;
 
+	// TODO: Some sequences do not evaluate correctly with zero timescale if this flag
+	// is unset, but it's not clear why or which ones. Workaround for UE-146012.
+	HoldSection->Parameters.bCanLoop = true;
+
 	ItemState.WeakHoldSection = HoldSection;
 
 	return true;
diff --git a/Engine/Plugins/VirtualProduction/SequencerPlaylists/Source/SequencerPlaylists/Private/SequencerPlaylistsModule.h b/Engine/Plugins/VirtualProduction/SequencerPlaylists/Source/SequencerPlaylists/Private/SequencerPlaylistsModule.h
--- a/Engine/Plugins/VirtualProduction/SequencerPlaylists/Source/SequencerPlaylists/Private/SequencerPlaylistsModule.h
+++ b/Engine/Plugins/VirtualProduction/SequencerPlaylists/Source/SequencerPlaylists/Private/SequencerPlaylistsModule.h
@@ -4,7 +4,6 @@
 
 #include "CoreMinimal.h"
 #include "Templates/SubclassOf.h"
-#include "UObject/StrongObjectPtr.h"
 #include "ISequencerPlaylistsModule.h"
 #include "SequencerPlaylistItem.h"
 
@@ -47,8 +46,6 @@ private:
 private:
 	TMap<TSubclassOf<USequencerPlaylistItem>, FSequencerPlaylistItemPlayerFactory> ItemPlayerFactories;
 
-	TStrongObjectPtr<USequencerPlaylistPlayer> DefaultPlayer;
-
 	TSharedPtr<FUICommandList> PluginCommands;
 
 	TSharedPtr<FAssetTypeActions_SequencerPlaylist> AssetTypeActions;
