Sequencer Playlists: Backing out workaround from CL 22248736, which recently regressed and now leads to a crash.

diff --git a/Engine/Plugins/VirtualProduction/SequencerPlaylists/Source/SequencerPlaylists/Private/SequencerPlaylistItem_Sequence.cpp b/Engine/Plugins/VirtualProduction/SequencerPlaylists/Source/SequencerPlaylists/Private/SequencerPlaylistItem_Sequence.cpp
--- a/Engine/Plugins/VirtualProduction/SequencerPlaylists/Source/SequencerPlaylists/Private/SequencerPlaylistItem_Sequence.cpp
+++ b/Engine/Plugins/VirtualProduction/SequencerPlaylists/Source/SequencerPlaylists/Private/SequencerPlaylistItem_Sequence.cpp
@@ -234,10 +234,6 @@ bool FSequencerPlaylistItemPlayer_Sequence::AddHold(USequencerPlaylistItem* Item
 
 	HoldSection->Parameters.TimeScale = 0.f;
 
-	// TODO: Some sequences do not evaluate correctly with zero timescale if this flag
-	// is unset, but it's not clear why or which ones. Workaround for UE-146012.
-	HoldSection->Parameters.bCanLoop = true;
-
 	ItemState.WeakHoldSection = HoldSection;
 
 	return true;
