Sequencer Playlists: Additional playback options.

- Invoking play again before previous section(s) have completed now enqueues additional sections on the same row (instead of playing concurrently).
- Replaced "hold" with "pause."
- Added "play backwards."

diff --git a/Engine/Plugins/VirtualProduction/SequencerPlaylists/Source/SequencerPlaylists/Private/SequencerPlaylistItem_Sequence.cpp b/Engine/Plugins/VirtualProduction/SequencerPlaylists/Source/SequencerPlaylists/Private/SequencerPlaylistItem_Sequence.cpp
--- a/Engine/Plugins/VirtualProduction/SequencerPlaylists/Source/SequencerPlaylists/Private/SequencerPlaylistItem_Sequence.cpp
+++ b/Engine/Plugins/VirtualProduction/SequencerPlaylists/Source/SequencerPlaylists/Private/SequencerPlaylistItem_Sequence.cpp
@@ -2,11 +2,11 @@
 
 #include "SequencerPlaylistItem_Sequence.h"
 #include "Evaluation/MovieSceneSequenceTransform.h"
-#include "SequencerPlaylist.h"
-
 #include "Framework/Notifications/NotificationManager.h"
 #include "LevelSequence.h"
 #include "MovieSceneFolder.h"
+#include "Recorder/TakeRecorder.h"
+#include "SequencerPlaylist.h"
 #include "TrackEditors/SubTrackEditorBase.h" // for FSubTrackEditorUtil::CanAddSubSequence
 #include "Tracks/MovieSceneSubTrack.h"
 #include "Widgets/Notifications/SNotificationList.h"
@@ -15,6 +15,56 @@
 #define LOCTEXT_NAMESPACE "SequencerPlaylists"
 
 
+namespace UE::Private::SequencerPlaylistItem_Sequence
+{
+	bool IsTimeWithinSection(FQualifiedFrameTime InTime, UMovieSceneSection* Section)
+	{
+		check(Section);
+
+		UMovieScene* SectionScene = Section->GetTypedOuter<UMovieScene>();
+		check(SectionScene);
+
+		const FFrameTime Time_SceneTicks = InTime.ConvertTo(SectionScene->GetTickResolution());
+		return Section->IsTimeWithinSection(Time_SceneTicks.FloorToFrame());
+	}
+
+
+	UMovieSceneSubSection* IsTimeWithinAnySection(FQualifiedFrameTime InTime, TArrayView<TWeakObjectPtr<UMovieSceneSubSection>> WeakSectionArray)
+	{
+		for (TWeakObjectPtr<UMovieSceneSubSection> WeakSection : WeakSectionArray)
+		{
+			if (UMovieSceneSubSection* Section = WeakSection.Get())
+			{
+				if (IsTimeWithinSection(InTime, Section))
+				{
+					return Section;
+				}
+			}
+		}
+
+		return nullptr;
+	}
+
+
+	TPair<TOptional<FFrameTime>, FMovieSceneWarpCounter>
+	GetInnerTimeAndWarp(FQualifiedFrameTime QualifiedTime, UMovieSceneSubSection* SubSection)
+	{
+		check(SubSection);
+
+		TPair<TOptional<FFrameTime>, FMovieSceneWarpCounter> Result;
+
+		if (IsTimeWithinSection(QualifiedTime, SubSection))
+		{
+			FFrameTime InnerTime;
+			SubSection->OuterToInnerTransform().TransformTime(QualifiedTime.Time, InnerTime, Result.Get<1>());
+			Result.Get<0>() = InnerTime;
+		}
+
+		return MoveTemp(Result);
+	}
+}
+
+
 FText USequencerPlaylistItem_Sequence::GetDisplayName()
 {
 	return Sequence ? Sequence->GetDisplayName() : LOCTEXT("SequenceItemNullDisplayName", "(No sequence)");
@@ -37,7 +87,76 @@ FSequencerPlaylistItemPlayer_Sequence::~FSequencerPlaylistItemPlayer_Sequence()
 }
 
 
-bool FSequencerPlaylistItemPlayer_Sequence::Play(USequencerPlaylistItem* Item)
+bool FSequencerPlaylistItemPlayer_Sequence::Play(
+	USequencerPlaylistItem* Item,
+	ESequencerPlaylistPlaybackDirection Direction // = ESequencerPlaylistPlaybackDirection::Forward
+)
+{
+	FPlayParams PlayParams;
+	PlayParams.Direction = Direction;
+	return InternalPlay(Item, PlayParams);
+}
+
+
+bool FSequencerPlaylistItemPlayer_Sequence::Pause(USequencerPlaylistItem* Item)
+{
+	return InternalPause(Item);
+}
+
+
+bool FSequencerPlaylistItemPlayer_Sequence::Stop(USequencerPlaylistItem* Item)
+{
+	return InternalStop(Item);
+}
+
+
+bool FSequencerPlaylistItemPlayer_Sequence::AddHold(USequencerPlaylistItem* Item)
+{
+	return InternalAddHold(Item, FHoldParams());
+}
+
+
+bool FSequencerPlaylistItemPlayer_Sequence::Reset(USequencerPlaylistItem* Item)
+{
+	return InternalReset(Item);
+}
+
+
+FSequencerPlaylistPlaybackState
+FSequencerPlaylistItemPlayer_Sequence::GetPlaybackState(USequencerPlaylistItem* Item) const
+{
+	FSequencerPlaylistPlaybackState Result;
+
+	USequencerPlaylistItem_Sequence* SequenceItem = CastChecked<USequencerPlaylistItem_Sequence>(Item);
+	if (!SequenceItem || !SequenceItem->Sequence)
+	{
+		return Result;
+	}
+
+	TSharedPtr<ISequencer> Sequencer = WeakSequencer.Pin();
+	if (!ensure(Sequencer))
+	{
+		return Result;
+	}
+
+	const FItemState* ItemState = ItemStates.Find(Item);
+	if (!ItemState)
+	{
+		return Result;
+	}
+
+	Result.bIsPlaying = ItemState->PlayingUntil_RootTicks > Sequencer->GetGlobalTime().Time.FloorToFrame().Value;
+	Result.bIsPaused = ItemState->WeakHoldSection.IsValid();
+	Result.PlaybackDirection = ItemState->LastPlayDirection;
+
+	return Result;
+}
+
+
+bool FSequencerPlaylistItemPlayer_Sequence::InternalPlay(
+	USequencerPlaylistItem* Item,
+	const FPlayParams& PlayParams
+)
 {
 	USequencerPlaylistItem_Sequence* SequenceItem = CastChecked<USequencerPlaylistItem_Sequence>(Item);
 	if (!SequenceItem || !SequenceItem->Sequence)
@@ -70,34 +189,37 @@ bool FSequencerPlaylistItemPlayer_Sequence::Play(USequencerPlaylistItem* Item)
 	if (UMovieSceneSubSection* HoldSection = ItemState.WeakHoldSection.Get())
 	{
 		EndSection(HoldSection);
+		TruncatePlayingUntil(ItemState);
 		ItemState.WeakHoldSection.Reset();
 	}
 
-	UMovieScene* PlayScene = SequenceItem->Sequence->GetMovieScene();
-	TRange<FFrameNumber> PlayRange = PlayScene->GetPlaybackRange();
+	ItemState.LastPlayDirection = PlayParams.Direction;
 
-	const FFrameTime StartFrameOffset_Ticks = ConvertFrameTime(FFrameNumber(SequenceItem->StartFrameOffset),
-		PlayScene->GetDisplayRate(), PlayScene->GetTickResolution());
+	UMovieScene* PlayScene = SequenceItem->Sequence->GetMovieScene();
+	const TRange<FFrameNumber> SequencePlayRange = PlayScene->GetPlaybackRange();
 
-	const FFrameTime EndFrameOffset_Ticks = ConvertFrameTime(FFrameNumber(SequenceItem->EndFrameOffset),
-		PlayScene->GetDisplayRate(), PlayScene->GetTickResolution());
+	const FFrameTime StartFrameOffset_Ticks = PlayParams.StartFrameOffset_SceneTicks.IsSet()
+		? PlayParams.StartFrameOffset_SceneTicks.GetValue()
+		: ConvertFrameTime(FFrameNumber(SequenceItem->StartFrameOffset),
+		                   PlayScene->GetDisplayRate(), PlayScene->GetTickResolution());
 
-	if (SequenceItem->StartFrameOffset > 0)
-	{
-		PlayRange.SetLowerBoundValue(PlayRange.GetLowerBoundValue() + StartFrameOffset_Ticks.FloorToFrame());
-	}
+	const FFrameTime EndFrameOffset_Ticks = PlayParams.EndFrameOffset_SceneTicks.IsSet()
+		? PlayParams.EndFrameOffset_SceneTicks.GetValue()
+		: ConvertFrameTime(FFrameNumber(SequenceItem->EndFrameOffset),
+		                   PlayScene->GetDisplayRate(), PlayScene->GetTickResolution());
 
-	if (SequenceItem->EndFrameOffset > 0)
-	{
-		PlayRange.SetUpperBoundValue(PlayRange.GetUpperBoundValue() - EndFrameOffset_Ticks.FloorToFrame());
-	}
+	TRange<FFrameNumber> SectionRange(
+		SequencePlayRange.GetLowerBoundValue() + StartFrameOffset_Ticks.FloorToFrame(),
+		SequencePlayRange.GetUpperBoundValue() - EndFrameOffset_Ticks.FloorToFrame());
 
 	const float TimeScale = FMath::Max(SMALL_NUMBER, SequenceItem->PlaybackSpeed);
-	const FFrameTime SingleLoopDuration = ConvertFrameTime(PlayRange.Size<FFrameTime>() / TimeScale,
+	const FFrameTime SingleLoopDuration = ConvertFrameTime(SectionRange.Size<FFrameTime>() / TimeScale,
 		PlayScene->GetTickResolution(), RootScene->GetTickResolution());
 
-	const FQualifiedFrameTime GlobalTime = Sequencer->GetGlobalTime();
-	const FFrameNumber StartFrame = GlobalTime.Time.FloorToFrame();
+	const FFrameNumber NowFrame = Sequencer->GetGlobalTime().Time.FloorToFrame();
+	const FFrameNumber EnqueueFrame = FMath::Max(ItemState.PlayingUntil_RootTicks, NowFrame.Value);
+	const FFrameNumber StartFrame = PlayParams.bEnqueueExtraPlays ? EnqueueFrame : NowFrame;
+
 	const int32 MinDuration = 1;
 	const int32 MaxDuration = TNumericLimits<int32>::Max() - StartFrame.Value - 1;
 	const int32 Duration = FMath::Clamp(
@@ -109,29 +231,94 @@ bool FSequencerPlaylistItemPlayer_Sequence::Play(USequencerPlaylistItem* Item)
 	UMovieSceneSubSection* WorkingSubSection = WorkingTrack->AddSequence(SequenceItem->Sequence, StartFrame, Duration);
 
 	WorkingSubSection->Parameters.TimeScale = TimeScale;
+	WorkingSubSection->Parameters.StartFrameOffset = StartFrameOffset_Ticks.FloorToFrame();
+	WorkingSubSection->Parameters.EndFrameOffset = EndFrameOffset_Ticks.FloorToFrame();
 
-	if (SequenceItem->StartFrameOffset > 0)
+	if (SequenceItem->NumLoops != 0)
 	{
-		WorkingSubSection->Parameters.StartFrameOffset = StartFrameOffset_Ticks.FloorToFrame();
+		WorkingSubSection->Parameters.bCanLoop = true;
 	}
 
-	if (SequenceItem->EndFrameOffset > 0)
+	if (PlayParams.Direction == ESequencerPlaylistPlaybackDirection::Reverse)
 	{
-		WorkingSubSection->Parameters.EndFrameOffset = EndFrameOffset_Ticks.FloorToFrame();
+		WorkingSubSection->Parameters.TimeScale *= -1.0f;
+		WorkingSubSection->Parameters.StartFrameOffset =
+			SequencePlayRange.Size<FFrameNumber>() - WorkingSubSection->Parameters.StartFrameOffset;
 	}
 
-	if (SequenceItem->NumLoops != 0)
+	ItemState.WeakPlaySections.Add(WorkingSubSection);
+
+	return true;
+}
+
+
+bool FSequencerPlaylistItemPlayer_Sequence::InternalPause(USequencerPlaylistItem* Item)
+{
+	using namespace UE::Private::SequencerPlaylistItem_Sequence;
+
+	bool bSequenceWasModified = false;
+
+	USequencerPlaylistItem_Sequence* SequenceItem = CastChecked<USequencerPlaylistItem_Sequence>(Item);
+	if (!SequenceItem || !SequenceItem->Sequence)
 	{
-		WorkingSubSection->Parameters.bCanLoop = true;
+		return bSequenceWasModified;
 	}
 
-	ItemState.WeakPlaySections.Add(WorkingSubSection);
+	TSharedPtr<ISequencer> Sequencer = WeakSequencer.Pin();
+	if (!ensure(Sequencer))
+	{
+		return bSequenceWasModified;
+	}
 
-	return true;
+	const FQualifiedFrameTime GlobalTime = Sequencer->GetGlobalTime();
+
+	FItemState& ItemState = ItemStates.FindOrAdd(Item);
+	UMovieSceneSubSection* HoldSection = ItemState.WeakHoldSection.Get();
+	if (HoldSection)
+	{
+		if (IsSequencerRecordingOrPlaying())
+		{
+			// Resume playback
+			FPlayParams PlayParams;
+			PlayParams.Direction = ItemState.LastPlayDirection;
+			PlayParams.StartFrameOffset_SceneTicks = HoldSection->Parameters.StartFrameOffset;
+			bSequenceWasModified |= InternalPlay(Item, PlayParams);
+		}
+		else
+		{
+			// Cancel hold
+			bSequenceWasModified |= Stop(Item);
+			ItemState.bHoldOnReset = false;
+		}
+	}
+	else
+	{
+		FHoldParams HoldParams;
+		if (IsSequencerRecordingOrPlaying())
+		{
+			UMovieSceneSubSection* CurrentPlayingSection = IsTimeWithinAnySection(GlobalTime, ItemState.WeakPlaySections);
+			if (CurrentPlayingSection)
+			{
+				TPair<TOptional<FFrameTime>, FMovieSceneWarpCounter> InnerTimeAndWarp =
+					GetInnerTimeAndWarp(GlobalTime, CurrentPlayingSection);
+
+				HoldParams.StartFrameOffset_SceneTicks = InnerTimeAndWarp.Get<0>().GetValue();
+			}
+		}
+		else
+		{
+			ItemState.bHoldOnReset = true;
+		}
+
+		bSequenceWasModified |= Stop(Item);
+		bSequenceWasModified |= InternalAddHold(Item, HoldParams);
+	}
+
+	return bSequenceWasModified;
 }
 
 
-bool FSequencerPlaylistItemPlayer_Sequence::Stop(USequencerPlaylistItem* Item)
+bool FSequencerPlaylistItemPlayer_Sequence::InternalStop(USequencerPlaylistItem* Item)
 {
 	bool bSequenceWasModified = false;
 
@@ -164,18 +351,13 @@ bool FSequencerPlaylistItemPlayer_Sequence::Stop(USequencerPlaylistItem* Item)
 	}
 
 	ItemState.WeakPlaySections.Empty();
-
-	if (TSharedPtr<ISequencer> Sequencer = WeakSequencer.Pin())
-	{
-		const FFrameNumber Now = Sequencer->GetGlobalTime().Time.FloorToFrame();
-		ItemState.PlayingUntil_RootTicks = FMath::Min(Now.Value, ItemState.PlayingUntil_RootTicks);
-	}
+	TruncatePlayingUntil(ItemState);
 
 	return bSequenceWasModified;
 }
 
 
-bool FSequencerPlaylistItemPlayer_Sequence::AddHold(USequencerPlaylistItem* Item)
+bool FSequencerPlaylistItemPlayer_Sequence::InternalAddHold(USequencerPlaylistItem* Item, const FHoldParams& HoldParams)
 {
 	USequencerPlaylistItem_Sequence* SequenceItem = CastChecked<USequencerPlaylistItem_Sequence>(Item);
 	if (!SequenceItem || !SequenceItem->Sequence)
@@ -220,14 +402,12 @@ bool FSequencerPlaylistItemPlayer_Sequence::AddHold(USequencerPlaylistItem* Item
 	HoldSection = WorkingTrack->AddSequence(HoldSequence, StartFrame, MaxDuration);
 	ItemState.PlayingUntil_RootTicks = TNumericLimits<int32>::Max();
 
-	if (SequenceItem->StartFrameOffset > 0)
-	{
-		const FFrameTime StartFrameOffset_Ticks = ConvertFrameTime(
-			FFrameNumber(SequenceItem->StartFrameOffset),
-			HoldScene->GetDisplayRate(),
-			HoldScene->GetTickResolution());
-		HoldSection->Parameters.StartFrameOffset = StartFrameOffset_Ticks.FloorToFrame();
-	}
+	const FFrameTime StartFrameOffset_Ticks = HoldParams.StartFrameOffset_SceneTicks.IsSet()
+		? HoldParams.StartFrameOffset_SceneTicks.GetValue()
+		: ConvertFrameTime(FFrameNumber(SequenceItem->StartFrameOffset),
+		                   HoldScene->GetDisplayRate(), HoldScene->GetTickResolution());
+
+	HoldSection->Parameters.StartFrameOffset = StartFrameOffset_Ticks.FloorToFrame();
 
 	HoldSection->Parameters.TimeScale = 0.f;
 
@@ -237,26 +417,31 @@ bool FSequencerPlaylistItemPlayer_Sequence::AddHold(USequencerPlaylistItem* Item
 }
 
 
-bool FSequencerPlaylistItemPlayer_Sequence::Reset(USequencerPlaylistItem* Item)
+bool FSequencerPlaylistItemPlayer_Sequence::InternalReset(USequencerPlaylistItem* Item)
 {
-	if (Item->bHoldAtFirstFrame)
+	bool bSequenceWasModified = false;
+
+	FItemState& ItemState = ItemStates.FindOrAdd(Item);
+	if (ItemState.bHoldOnReset)
 	{
-		Stop(Item);
-		return AddHold(Item);
+		bSequenceWasModified |= Stop(Item);
+		bSequenceWasModified |= AddHold(Item);
 	}
 	else
 	{
-		return Stop(Item);
+		bSequenceWasModified |= Stop(Item);
 	}
+
+	return bSequenceWasModified;
 }
 
 
-bool FSequencerPlaylistItemPlayer_Sequence::IsPlaying(USequencerPlaylistItem* Item) const
+bool FSequencerPlaylistItemPlayer_Sequence::IsSequencerRecordingOrPlaying() const
 {
-	USequencerPlaylistItem_Sequence* SequenceItem = CastChecked<USequencerPlaylistItem_Sequence>(Item);
-	if (!SequenceItem || !SequenceItem->Sequence)
+	UTakeRecorder* Recorder = UTakeRecorder::GetActiveRecorder();
+	if (Recorder && Recorder->GetState() == ETakeRecorderState::Started)
 	{
-		return false;
+		return true;
 	}
 
 	TSharedPtr<ISequencer> Sequencer = WeakSequencer.Pin();
@@ -265,13 +450,17 @@ bool FSequencerPlaylistItemPlayer_Sequence::IsPlaying(USequencerPlaylistItem* It
 		return false;
 	}
 
-	const FItemState* ItemState = ItemStates.Find(Item);
-	if (!ItemState)
+	return Sequencer->GetPlaybackStatus() == EMovieScenePlayerStatus::Playing;
+}
+
+
+void FSequencerPlaylistItemPlayer_Sequence::TruncatePlayingUntil(FItemState& InItemState)
+{
+	if (TSharedPtr<ISequencer> Sequencer = WeakSequencer.Pin())
 	{
-		return false;
+		const FFrameNumber Now = Sequencer->GetGlobalTime().Time.FloorToFrame();
+		InItemState.PlayingUntil_RootTicks = FMath::Min(Now.Value, InItemState.PlayingUntil_RootTicks);
 	}
-
-	return ItemState->PlayingUntil_RootTicks > Sequencer->GetGlobalTime().Time.FloorToFrame().Value;
 }
 
 
@@ -333,7 +522,7 @@ UMovieSceneSubTrack* FSequencerPlaylistItemPlayer_Sequence::GetOrCreateWorkingTr
 	{
 		FolderToUse = NewObject<UMovieSceneFolder>(RootScene, NAME_None, RF_Transactional);
 		FolderToUse->SetFolderName(PlaylistFolderName);
-		//RootScene->GetRootFolders().Add(FolderToUse);
+		RootScene->AddRootFolder(FolderToUse);
 	}
 
 	FolderToUse->AddChildTrack(NewWorkingTrack);
@@ -345,6 +534,8 @@ UMovieSceneSubTrack* FSequencerPlaylistItemPlayer_Sequence::GetOrCreateWorkingTr
 
 bool FSequencerPlaylistItemPlayer_Sequence::EndSection(UMovieSceneSection* Section)
 {
+	using namespace UE::Private::SequencerPlaylistItem_Sequence;
+
 	bool bSequenceWasModified = false;
 
 	if (!ensure(Section))
@@ -364,13 +555,28 @@ bool FSequencerPlaylistItemPlayer_Sequence::EndSection(UMovieSceneSection* Secti
 	check(SectionScene);
 
 	const FQualifiedFrameTime GlobalTime = Sequencer->GetGlobalTime();
-	const FFrameTime SectionNow = GlobalTime.ConvertTo(SectionScene->GetTickResolution());
-	if (Section->IsTimeWithinSection(SectionNow.FloorToFrame()))
+	const FFrameTime Now_SceneTicks = GlobalTime.ConvertTo(SectionScene->GetTickResolution());
+	const FFrameTime SectionStart_SceneTicks = Section->GetInclusiveStartFrame();
+	if (SectionStart_SceneTicks > Now_SceneTicks)
+	{
+		// Remove queued sections entirely (start beyond the current time)
+		if (UMovieSceneTrack* SectionTrack = Section->GetTypedOuter<UMovieSceneTrack>())
+		{
+			RootSequence->Modify();
+			bSequenceWasModified = true;
+
+			SectionTrack->Modify();
+			SectionTrack->RemoveSection(*Section);
+			return bSequenceWasModified;
+		}
+	}
+	else if (IsTimeWithinSection(GlobalTime, Section))
 	{
+		// Trim sections the current time intersects
 		RootSequence->Modify();
 		bSequenceWasModified = true;
 
-		Section->SetEndFrame(TRangeBound<FFrameNumber>::Exclusive(SectionNow.FloorToFrame()));
+		Section->SetEndFrame(TRangeBound<FFrameNumber>::Exclusive(Now_SceneTicks.FloorToFrame()));
 
 		// Remove degenerate sections. This can happen if we reset then stop a held item while paused.
 		if (Section->GetRange().IsEmpty())
@@ -389,11 +595,9 @@ bool FSequencerPlaylistItemPlayer_Sequence::EndSection(UMovieSceneSection* Secti
 			if (SubSection->Parameters.bCanLoop)
 			{
 				// Calc whether we've looped; if not, set bCanLoop back to false.
-				FFrameTime InnerTime;
-				FMovieSceneWarpCounter WarpCounter;
-				SubSection->OuterToInnerTransform().TransformTime(GlobalTime.Time, InnerTime, WarpCounter);
-				const uint32 LastWarpCount = WarpCounter.LastWarpCount();
-				if (LastWarpCount == 0)
+				TPair<TOptional<FFrameTime>, FMovieSceneWarpCounter> InnerTimeAndWarp =
+					GetInnerTimeAndWarp(GlobalTime, SubSection);
+				if (InnerTimeAndWarp.Get<1>().LastWarpCount() == 0)
 				{
 					SubSection->Parameters.bCanLoop = false;
 				}
diff --git a/Engine/Plugins/VirtualProduction/SequencerPlaylists/Source/SequencerPlaylists/Private/SequencerPlaylistItem_Sequence.h b/Engine/Plugins/VirtualProduction/SequencerPlaylists/Source/SequencerPlaylists/Private/SequencerPlaylistItem_Sequence.h
--- a/Engine/Plugins/VirtualProduction/SequencerPlaylists/Source/SequencerPlaylists/Private/SequencerPlaylistItem_Sequence.h
+++ b/Engine/Plugins/VirtualProduction/SequencerPlaylists/Source/SequencerPlaylists/Private/SequencerPlaylistItem_Sequence.h
@@ -3,6 +3,7 @@
 #pragma once
 
 #include "ISequencerPlaylistsModule.h"
+#include "Misc/FrameTime.h"
 #include "SequencerPlaylistItem.h"
 #include "SequencerPlaylistItem_Sequence.generated.h"
 
@@ -22,7 +23,7 @@ class USequencerPlaylistItem_Sequence : public USequencerPlaylistItem
 	FText GetDisplayName() override;
 
 public:
-	UPROPERTY(EditAnywhere, BlueprintReadWrite, Category="SequencerPlaylists", meta=(NoResetToDefault))
+	UPROPERTY(EditAnywhere, BlueprintReadWrite, Category="Sequencer Playlists", meta=(NoResetToDefault))
 	TObjectPtr<ULevelSequence> Sequence;
 };
 
@@ -35,6 +36,8 @@ class FSequencerPlaylistItemPlayer_Sequence : public ISequencerPlaylistItemPlaye
 		TWeakObjectPtr<UMovieSceneSubSection> WeakHoldSection;
 		TArray<TWeakObjectPtr<UMovieSceneSubSection>> WeakPlaySections;
 		int32 PlayingUntil_RootTicks = TNumericLimits<int32>::Min();
+		ESequencerPlaylistPlaybackDirection LastPlayDirection = ESequencerPlaylistPlaybackDirection::Forward;
+		bool bHoldOnReset = false;
 	};
 
 public:
@@ -42,15 +45,46 @@ public:
 	~FSequencerPlaylistItemPlayer_Sequence() override;
 
 	//~ Begin ISequencerPlaylistItemPlayer
-	bool Play(USequencerPlaylistItem* Item) override;
+	bool Play(USequencerPlaylistItem* Item,
+	          ESequencerPlaylistPlaybackDirection Direction = ESequencerPlaylistPlaybackDirection::Forward) override;
+	bool Pause(USequencerPlaylistItem* Item) override;
 	bool Stop(USequencerPlaylistItem* Item) override;
 	bool AddHold(USequencerPlaylistItem* Item) override;
 	bool Reset(USequencerPlaylistItem* Item) override;
 
-	bool IsPlaying(USequencerPlaylistItem* Item) const override;
+	FSequencerPlaylistPlaybackState GetPlaybackState(USequencerPlaylistItem* Item) const override;
 	//~ End ISequencerPlaylistItemPlayer
 
 private:
+	struct FPlayParams
+	{
+		// If true, playing again before previous plays have completed will
+		// queue the additional play at the end.
+		// Otherwise, additional plays will occur concurrently on other rows.
+		bool bEnqueueExtraPlays = true;
+
+		ESequencerPlaylistPlaybackDirection Direction = ESequencerPlaylistPlaybackDirection::Forward;
+
+		// If not provided, defaults to the Item StartFrameOffset property.
+		TOptional<FFrameTime> StartFrameOffset_SceneTicks;
+
+		// If not provided, defaults to the Item EndFrameOffset property.
+		TOptional<FFrameTime> EndFrameOffset_SceneTicks;
+	};
+	bool InternalPlay(USequencerPlaylistItem* Item, const FPlayParams& PlayParams);
+	bool InternalPause(USequencerPlaylistItem* Item);
+	bool InternalStop(USequencerPlaylistItem* Item);
+
+	struct FHoldParams
+	{
+		// If not provided, defaults to the Item StartFrameOffset property.
+		TOptional<FFrameTime> StartFrameOffset_SceneTicks;
+	};
+	bool InternalAddHold(USequencerPlaylistItem* Item, const FHoldParams& HoldParams);
+	bool InternalReset(USequencerPlaylistItem* Item);
+
+	bool IsSequencerRecordingOrPlaying() const;
+	void TruncatePlayingUntil(FItemState& InItemState);
 	void ClearItemStates();
 
 	UMovieSceneSubTrack* GetOrCreateWorkingTrack(USequencerPlaylistItem* Item);
diff --git a/Engine/Plugins/VirtualProduction/SequencerPlaylists/Source/SequencerPlaylists/Private/SequencerPlaylistPlayer.cpp b/Engine/Plugins/VirtualProduction/SequencerPlaylists/Source/SequencerPlaylists/Private/SequencerPlaylistPlayer.cpp
--- a/Engine/Plugins/VirtualProduction/SequencerPlaylists/Source/SequencerPlaylists/Private/SequencerPlaylistPlayer.cpp
+++ b/Engine/Plugins/VirtualProduction/SequencerPlaylists/Source/SequencerPlaylists/Private/SequencerPlaylistPlayer.cpp
@@ -58,7 +58,10 @@ void USequencerPlaylistPlayer::SetPlaylist(USequencerPlaylist* InPlaylist)
 }
 
 
-bool USequencerPlaylistPlayer::PlayItem(USequencerPlaylistItem* Item)
+bool USequencerPlaylistPlayer::PlayItem(
+	USequencerPlaylistItem* Item,
+	ESequencerPlaylistPlaybackDirection Direction // = ESequencerPlaylistPlaybackDirection::Forward
+)
 {
 	if (!Item)
 	{
@@ -69,7 +72,28 @@ bool USequencerPlaylistPlayer::PlayItem(USequencerPlaylistItem* Item)
 
 	EnterUnboundedPlayIfNotRecording();
 
-	if (GetCheckedItemPlayer(Item)->Play(Item))
+	if (GetCheckedItemPlayer(Item)->Play(Item, Direction))
+	{
+		GetSequencer()->NotifyMovieSceneDataChanged(EMovieSceneDataChangeType::RefreshAllImmediately);
+		return true;
+	}
+	else
+	{
+		return false;
+	}
+}
+
+
+bool USequencerPlaylistPlayer::PauseItem(USequencerPlaylistItem* Item)
+{
+	if (!Item)
+	{
+		return false;
+	}
+
+	FScopedTransaction Transaction(FText::Format(LOCTEXT("PauseItemTransaction", "Toggle pause of {0}"), Item->GetDisplayName()));
+
+	if (GetCheckedItemPlayer(Item)->Pause(Item))
 	{
 		GetSequencer()->NotifyMovieSceneDataChanged(EMovieSceneDataChangeType::RefreshAllImmediately);
 		return true;
@@ -124,18 +148,25 @@ bool USequencerPlaylistPlayer::ResetItem(USequencerPlaylistItem* Item)
 
 
 bool USequencerPlaylistPlayer::IsPlaying(USequencerPlaylistItem* Item)
+{
+	return GetPlaybackState(Item).bIsPlaying;
+}
+
+
+FSequencerPlaylistPlaybackState USequencerPlaylistPlayer::GetPlaybackState(USequencerPlaylistItem* Item)
 {
 	if (!Item)
 	{
-		return false;
+		return FSequencerPlaylistPlaybackState();
 	}
 
+	// If Sequencer isn't already open, don't open it
 	if (!WeakSequencer.IsValid())
 	{
-		return false;
+		return FSequencerPlaylistPlaybackState();
 	}
 
-	return GetCheckedItemPlayer(Item)->IsPlaying(Item);
+	return GetCheckedItemPlayer(Item)->GetPlaybackState(Item);
 }
 
 
@@ -257,7 +288,9 @@ void USequencerPlaylistPlayer::EnterUnboundedPlayIfNotRecording()
 	}
 }
 
-bool USequencerPlaylistPlayer::PlayAll()
+bool USequencerPlaylistPlayer::PlayAll(
+	ESequencerPlaylistPlaybackDirection Direction // = ESequencerPlaylistPlaybackDirection::Forward
+)
 {
 	if (!ensure(Playlist) || !Playlist->Items.Num())
 	{
@@ -277,7 +310,7 @@ bool USequencerPlaylistPlayer::PlayAll()
 			continue;
 		}
 
-		bAnyChange |= GetCheckedItemPlayer(Item)->Play(Item);
+		bAnyChange |= GetCheckedItemPlayer(Item)->Play(Item, Direction);
 	}
 
 	if (bAnyChange)
@@ -443,23 +476,6 @@ void USequencerPlaylistPlayer::OnTakeRecorderInitialized(UTakeRecorder* InRecord
 
 void USequencerPlaylistPlayer::OnTakeRecorderStarted(UTakeRecorder* InRecorder)
 {
-	if (!ensure(Playlist))
-	{
-		return;
-	}
-
-	if (TSharedPtr<ISequencer> Sequencer = GetValidatedSequencer())
-	{
-		FScopedTransaction Transaction(LOCTEXT("TakeRecorderStartedTransaction", "Playlist - Take Recorder started"));
-
-		for (USequencerPlaylistItem* Item : Playlist->Items)
-		{
-			if (Item->bHoldAtFirstFrame && !Item->bMute)
-			{
-				GetCheckedItemPlayer(Item)->AddHold(Item);
-			}
-		}
-	}
 }
 
 
diff --git a/Engine/Plugins/VirtualProduction/SequencerPlaylists/Source/SequencerPlaylists/Private/SequencerPlaylistsStyle.cpp b/Engine/Plugins/VirtualProduction/SequencerPlaylists/Source/SequencerPlaylists/Private/SequencerPlaylistsStyle.cpp
--- a/Engine/Plugins/VirtualProduction/SequencerPlaylists/Source/SequencerPlaylists/Private/SequencerPlaylistsStyle.cpp
+++ b/Engine/Plugins/VirtualProduction/SequencerPlaylists/Source/SequencerPlaylists/Private/SequencerPlaylistsStyle.cpp
@@ -128,10 +128,10 @@ TSharedRef<FSlateStyleSet> FSequencerPlaylistsStyle::Create()
 
 	const FLinearColor PlayColor = FStyleColors::AccentGreen.GetSpecifiedColor();
 	const FLinearColor StopColor = FStyleColors::AccentRed.GetSpecifiedColor();
-	const FLinearColor ResetColor = FStyleColors::AccentBlue.GetSpecifiedColor();
+	const FLinearColor PauseResetColor = FStyleColors::AccentBlue.GetSpecifiedColor();
 	const FLinearColor PlayPressedColor = MakeColorVariation(PlayColor, EColorVariation::Dimmed);
 	const FLinearColor StopPressedColor = MakeColorVariation(StopColor, EColorVariation::Dimmed);
-	const FLinearColor ResetPressedColor = MakeColorVariation(ResetColor, EColorVariation::Dimmed);
+	const FLinearColor PauseResetPressedColor = MakeColorVariation(PauseResetColor, EColorVariation::Dimmed);
 
 	const float CornerRadius = 4.0f;
 	const FVector4 AllCorners(CornerRadius, CornerRadius, CornerRadius, CornerRadius);
@@ -186,7 +186,15 @@ TSharedRef<FSlateStyleSet> FSequencerPlaylistsStyle::Create()
 		.SetNormalPadding(FMargin(6.f, 4.f, 8.f, 4.f))
 		.SetPressedPadding(FMargin(6.f, 4.f, 8.f, 4.f));
 
-	const FButtonStyle PlayTransportButton = FButtonStyle(LeftTransportButton)
+	const FButtonStyle PlayReverseTransportButton = FButtonStyle(LeftTransportButton)
+		.SetHoveredForeground(PlayColor)
+		.SetPressedForeground(PlayPressedColor);
+
+	const FButtonStyle PauseTransportButton = FButtonStyle(MiddleTransportButton)
+		.SetHoveredForeground(PauseResetColor)
+		.SetPressedForeground(PauseResetPressedColor);
+
+	const FButtonStyle PlayTransportButton = FButtonStyle(MiddleTransportButton)
 		.SetHoveredForeground(PlayColor)
 		.SetPressedForeground(PlayPressedColor);
 
@@ -195,9 +203,11 @@ TSharedRef<FSlateStyleSet> FSequencerPlaylistsStyle::Create()
 		.SetPressedForeground(StopPressedColor);
 
 	const FButtonStyle ResetTransportButton = FButtonStyle(RightTransportButton)
-		.SetHoveredForeground(ResetColor)
-		.SetPressedForeground(ResetPressedColor);
+		.SetHoveredForeground(PauseResetColor)
+		.SetPressedForeground(PauseResetPressedColor);
 
+	Style->Set("SequencerPlaylists.TransportButton.PlayReverse", PlayReverseTransportButton);
+	Style->Set("SequencerPlaylists.TransportButton.Pause", PauseTransportButton);
 	Style->Set("SequencerPlaylists.TransportButton.Play", PlayTransportButton);
 	Style->Set("SequencerPlaylists.TransportButton.Stop", StopTransportButton);
 	Style->Set("SequencerPlaylists.TransportButton.Reset", ResetTransportButton);
@@ -214,17 +224,27 @@ TSharedRef<FSlateStyleSet> FSequencerPlaylistsStyle::Create()
 		.SetHoveredForeground(PlayColor)
 		.SetPressedForeground(PlayPressedColor);
 
+	const FButtonStyle PauseHoverTransportButton = FButtonStyle(HoverTransportButtonStyle)
+		.SetHoveredForeground(PauseResetColor)
+		.SetPressedForeground(PauseResetPressedColor);
+
 	const FButtonStyle StopHoverTransportButton = FButtonStyle(HoverTransportButtonStyle)
 		.SetHoveredForeground(StopColor)
 		.SetPressedForeground(StopPressedColor);
 
 	const FButtonStyle ResetHoverTransportButton = FButtonStyle(HoverTransportButtonStyle)
-		.SetHoveredForeground(ResetColor)
-		.SetPressedForeground(ResetPressedColor);
+		.SetHoveredForeground(PauseResetColor)
+		.SetPressedForeground(PauseResetPressedColor);
+
+	const FButtonStyle PlayReverseHoverTransportButton = FButtonStyle(HoverTransportButtonStyle)
+		.SetHoveredForeground(PlayColor)
+		.SetPressedForeground(PlayPressedColor);
 
 	Style->Set("SequencerPlaylists.HoverTransport.Play", PlayHoverTransportButton);
+	Style->Set("SequencerPlaylists.HoverTransport.Pause", PauseHoverTransportButton);
 	Style->Set("SequencerPlaylists.HoverTransport.Stop", StopHoverTransportButton);
 	Style->Set("SequencerPlaylists.HoverTransport.Reset", ResetHoverTransportButton);
+	Style->Set("SequencerPlaylists.HoverTransport.PlayReverse", PlayReverseHoverTransportButton);
 
 	FEditableTextBoxStyle EditableTextStyle = FEditableTextBoxStyle()
 		.SetTextStyle(FAppStyle::Get().GetWidgetStyle<FTextBlockStyle>("NormalText"))
@@ -259,6 +279,10 @@ TSharedRef<FSlateStyleSet> FSequencerPlaylistsStyle::Create()
 
 			Style->Set("SequencerPlaylists.Play", new IMAGE_BRUSH_SVG("Starship/Common/play", Icon20x20));
 			Style->Set("SequencerPlaylists.Play.Small", new IMAGE_BRUSH_SVG("Starship/Common/play", Icon12x12));
+			Style->Set("SequencerPlaylists.Pause", new IMAGE_BRUSH_SVG("Starship/MainToolbar/pause", Icon20x20)); // FIXME
+			Style->Set("SequencerPlaylists.Pause.Small", new IMAGE_BRUSH_SVG("Starship/MainToolbar/pause", Icon12x12)); // FIXME
+			Style->Set("SequencerPlaylists.PlayReverse", new IMAGE_BRUSH("Sequencer/Transport_Bar/Backwards_24x", Icon20x20)); // FIXME
+			Style->Set("SequencerPlaylists.PlayReverse.Small", new IMAGE_BRUSH("Sequencer/Transport_Bar/Backwards_24x", Icon12x12)); // FIXME
 			Style->Set("SequencerPlaylists.Stop", new IMAGE_BRUSH_SVG("Starship/Common/stop", Icon20x20));
 			Style->Set("SequencerPlaylists.Stop.Small", new IMAGE_BRUSH_SVG("Starship/Common/stop", Icon12x12));
 			Style->Set("SequencerPlaylists.Reset", new IMAGE_BRUSH_SVG("Starship/Common/Reset", Icon20x20));
diff --git a/Engine/Plugins/VirtualProduction/SequencerPlaylists/Source/SequencerPlaylists/Private/SequencerPlaylistsWidgets.cpp b/Engine/Plugins/VirtualProduction/SequencerPlaylists/Source/SequencerPlaylists/Private/SequencerPlaylistsWidgets.cpp
--- a/Engine/Plugins/VirtualProduction/SequencerPlaylists/Source/SequencerPlaylists/Private/SequencerPlaylistsWidgets.cpp
+++ b/Engine/Plugins/VirtualProduction/SequencerPlaylists/Source/SequencerPlaylists/Private/SequencerPlaylistsWidgets.cpp
@@ -48,7 +48,7 @@ const float SSequencerPlaylistPanel::DefaultWidth = 300.0f;
 const FName SSequencerPlaylistPanel::ColumnName_HoverTransport(TEXT("HoverTransport"));
 const FName SSequencerPlaylistPanel::ColumnName_Items(TEXT("Items"));
 const FName SSequencerPlaylistPanel::ColumnName_Offset(TEXT("Offset"));
-const FName SSequencerPlaylistPanel::ColumnName_Hold(TEXT("Hold"));
+const FName SSequencerPlaylistPanel::ColumnName_Pause(TEXT("Pause"));
 const FName SSequencerPlaylistPanel::ColumnName_Loop(TEXT("Loop"));
 const FName SSequencerPlaylistPanel::ColumnName_HoverDetails(TEXT("HoverDetails"));
 
@@ -320,6 +320,38 @@ TSharedRef<SWidget> SSequencerPlaylistPanel::Construct_RightToolbar()
 TSharedRef<SWidget> SSequencerPlaylistPanel::Construct_Transport()
 {
 	return SNew(SHorizontalBox)
+		+ SHorizontalBox::Slot()
+		.AutoWidth()
+		[
+			SNew(SButton)
+			.ButtonStyle(FSequencerPlaylistsStyle::Get(), "SequencerPlaylists.TransportButton.PlayReverse")
+			.ContentPadding(FMargin(0.0f, 2.0f))
+			.HAlign(HAlign_Center)
+			.VAlign(VAlign_Center)
+			.OnClicked(this, &SSequencerPlaylistPanel::HandleClicked_PlayAllReverse)
+			.ToolTipText(LOCTEXT("PlayAllReverseButtonTooltip", "Play all items in reverse simultaneously."))
+			[
+				SNew(SImage)
+				.Image(FSequencerPlaylistsStyle::Get().GetBrush("SequencerPlaylists.PlayReverse"))
+				.ColorAndOpacity(FSlateColor::UseForeground())
+			]
+		]
+		+ SHorizontalBox::Slot()
+		.AutoWidth()
+		[
+			SNew(SButton)
+			.ButtonStyle(FSequencerPlaylistsStyle::Get(), "SequencerPlaylists.TransportButton.Pause")
+			.ContentPadding(FMargin(0.0f, 2.0f))
+			.HAlign(HAlign_Center)
+			.VAlign(VAlign_Center)
+			.OnClicked(this, &SSequencerPlaylistPanel::HandleClicked_PlayAll)
+			.ToolTipText(LOCTEXT("PlayAllReverseButtonTooltip", "Pause all items."))
+			[
+				SNew(SImage)
+				.Image(FSequencerPlaylistsStyle::Get().GetBrush("SequencerPlaylists.Pause"))
+				.ColorAndOpacity(FSlateColor::UseForeground())
+			]
+		]
 		+ SHorizontalBox::Slot()
 		.AutoWidth()
 		[
@@ -404,12 +436,24 @@ TSharedRef<SWidget> SSequencerPlaylistPanel::Construct_ItemListView()
 			{
 				return SNew(SSequencerPlaylistItemWidget, InData, OwnerTableView)
 					.PlayMode(this, &SSequencerPlaylistPanel::InPlayMode)
-					.IsPlaying_Lambda([this, WeakItem = InData->WeakItem]() {
+					.IsPlaying_Lambda([this, WeakItem = InData->WeakItem]() -> bool {
+						USequencerPlaylistItem* Item = WeakItem.Get();
+						USequencerPlaylistPlayer* Player = WeakPlayer.Get();
+						if (Item && Player)
+						{
+							return Player->GetPlaybackState(Item).bIsPlaying;
+						}
+						else
+						{
+							return false;
+						}
+					})
+					.IsPaused_Lambda([this, WeakItem = InData->WeakItem]() -> bool {
 						USequencerPlaylistItem* Item = WeakItem.Get();
 						USequencerPlaylistPlayer* Player = WeakPlayer.Get();
 						if (Item && Player)
 						{
-							return Player->IsPlaying(Item);
+							return Player->GetPlaybackState(Item).bIsPaused;
 						}
 						else
 						{
@@ -417,6 +461,8 @@ TSharedRef<SWidget> SSequencerPlaylistPanel::Construct_ItemListView()
 						}
 					})
 					.OnPlayClicked(this, &SSequencerPlaylistPanel::HandleClicked_Item_Play)
+					.OnPlayReverseClicked(this, &SSequencerPlaylistPanel::HandleClicked_Item_PlayReverse)
+					.OnPauseClicked(this, &SSequencerPlaylistPanel::HandleClicked_Item_Pause)
 					.OnStopClicked(this, &SSequencerPlaylistPanel::HandleClicked_Item_Stop)
 					.OnResetClicked(this, &SSequencerPlaylistPanel::HandleClicked_Item_Reset)
 					.OnRemoveClicked(this, &SSequencerPlaylistPanel::HandleClicked_Item_Remove)
@@ -435,8 +481,8 @@ TSharedRef<SWidget> SSequencerPlaylistPanel::Construct_ItemListView()
 			+ SHeaderRow::Column(ColumnName_Offset)
 				.DefaultLabel(LOCTEXT("ColumnLabelOffset", "Offset"))
 				.FillSized(45.0f)
-			+ SHeaderRow::Column(ColumnName_Hold)
-				.DefaultLabel(LOCTEXT("ColumnLabelHold", "Hold"))
+			+ SHeaderRow::Column(ColumnName_Pause)
+				.DefaultLabel(LOCTEXT("ColumnLabelPause", "Pause"))
 				.FillSized(35.0f)
 				.HAlignCell(HAlign_Center)
 			+ SHeaderRow::Column(ColumnName_Loop)
@@ -766,6 +812,19 @@ FReply SSequencerPlaylistPanel::HandleClicked_PlayAll()
 }
 
 
+FReply SSequencerPlaylistPanel::HandleClicked_PlayAllReverse()
+{
+	USequencerPlaylistPlayer* Player = WeakPlayer.Get();
+	if (!ensure(Player))
+	{
+		return FReply::Unhandled();
+	}
+
+	Player->PlayAll(ESequencerPlaylistPlaybackDirection::Reverse);
+	return FReply::Handled();
+}
+
+
 FReply SSequencerPlaylistPanel::HandleClicked_StopAll()
 {
 	USequencerPlaylistPlayer* Player = WeakPlayer.Get();
@@ -825,6 +884,32 @@ FReply SSequencerPlaylistPanel::HandleClicked_Item_Play(SSequencerPlaylistItemWi
 }
 
 
+FReply SSequencerPlaylistPanel::HandleClicked_Item_PlayReverse(SSequencerPlaylistItemWidget& ItemWidget)
+{
+	USequencerPlaylistPlayer* Player = WeakPlayer.Get();
+	if (!ensure(Player))
+	{
+		return FReply::Unhandled();
+	}
+
+	Player->PlayItem(ItemWidget.GetItem(), ESequencerPlaylistPlaybackDirection::Reverse);
+	return FReply::Handled();
+}
+
+
+FReply SSequencerPlaylistPanel::HandleClicked_Item_Pause(SSequencerPlaylistItemWidget& ItemWidget)
+{
+	USequencerPlaylistPlayer* Player = WeakPlayer.Get();
+	if (!ensure(Player))
+	{
+		return FReply::Unhandled();
+	}
+
+	Player->PauseItem(ItemWidget.GetItem());
+	return FReply::Handled();
+}
+
+
 FReply SSequencerPlaylistPanel::HandleClicked_Item_Stop(SSequencerPlaylistItemWidget& ItemWidget)
 {
 	USequencerPlaylistPlayer* Player = WeakPlayer.Get();
@@ -878,7 +963,6 @@ bool SSequencerPlaylistPanel::HandleItemDetailsIsPropertyVisible(const FProperty
 
 		PropertyNameToColumnName.Add("StartFrameOffset", ColumnName_Offset);
 		PropertyNameToColumnName.Add("EndFrameOffset", ColumnName_Offset);
-		PropertyNameToColumnName.Add("bHoldAtFirstFrame", ColumnName_Hold);
 		PropertyNameToColumnName.Add("NumLoops", ColumnName_Loop);
 
 		bMapInitialized = true;
@@ -993,8 +1077,10 @@ FSequencerPlaylistItemDragDropOp::~FSequencerPlaylistItemDragDropOp()
 
 
 const FText SSequencerPlaylistItemWidget::PlayItemTooltipText(LOCTEXT("PlayItemTooltip", "Play just this item."));
-const FText SSequencerPlaylistItemWidget::StopItemTooltipText(LOCTEXT("StopItemTooltip", "Stop all running sequences created by this item."));
-const FText SSequencerPlaylistItemWidget::ResetItemTooltipText(LOCTEXT("ResetItemTooltip", "Stop all running sequences created by this item and re-hold, if specified."));
+const FText SSequencerPlaylistItemWidget::PlayReverseItemTooltipText(LOCTEXT("PlayReverseItemTooltip", "Play just this item in reverse."));
+const FText SSequencerPlaylistItemWidget::PauseItemTooltipText(LOCTEXT("PauseItemTooltip", "Pause just this item. If not recording or playing back, creates a hold section for this item."));
+const FText SSequencerPlaylistItemWidget::StopItemTooltipText(LOCTEXT("StopItemTooltip", "Stop all sequences created by this item."));
+const FText SSequencerPlaylistItemWidget::ResetItemTooltipText(LOCTEXT("ResetItemTooltip", "Stop all sequences created by this item and re-hold, if specified."));
 
 
 void SSequencerPlaylistItemWidget::Construct(const FArguments& InArgs, TSharedPtr<FSequencerPlaylistRowData> InRowData, const TSharedRef<STableViewBase>& OwnerTableView)
@@ -1006,8 +1092,11 @@ void SSequencerPlaylistItemWidget::Construct(const FArguments& InArgs, TSharedPt
 
 	PlayMode = InArgs._PlayMode;
 	IsPlaying = InArgs._IsPlaying;
+	IsPaused = InArgs._IsPaused;
 
 	PlayClickedDelegate = InArgs._OnPlayClicked;
+	PlayReverseClickedDelegate = InArgs._OnPlayReverseClicked;
+	PauseClickedDelegate = InArgs._OnPauseClicked;
 	StopClickedDelegate = InArgs._OnStopClicked;
 	ResetClickedDelegate = InArgs._OnResetClicked;
 	RemoveClickedDelegate = InArgs._OnRemoveClicked;
@@ -1139,7 +1228,7 @@ TSharedRef<SWidget> SSequencerPlaylistItemWidget::GenerateWidgetForColumn(const
 	{
 		return SNew(SVerticalBox)
 			+ SVerticalBox::Slot()
-			.Padding(8.0f, 4.0f, 0.0f, 0.0f)
+			.Padding(8.0f, 4.0f, 0.0f, 2.0f)
 			[
 				SNew(SButton)
 				.ButtonStyle(FSequencerPlaylistsStyle::Get(), "SequencerPlaylists.HoverTransport.Play")
@@ -1156,7 +1245,24 @@ TSharedRef<SWidget> SSequencerPlaylistItemWidget::GenerateWidgetForColumn(const
 				]
 			]
 			+ SVerticalBox::Slot()
-			.Padding(8.0f, 0.0f, 0.0f, 0.0f)
+			.Padding(8.0f, 2.0f, 0.0f, 2.0f)
+			[
+				SNew(SButton)
+				.ButtonStyle(FSequencerPlaylistsStyle::Get(), "SequencerPlaylists.HoverTransport.PlayReverse")
+				.ContentPadding(0)
+				.HAlign(HAlign_Center)
+				.VAlign(VAlign_Center)
+				.Visibility_Lambda([this]() { return (IsHovered() && !InPlayMode()) ? EVisibility::Visible : EVisibility::Hidden; })
+				.OnClicked_Lambda([this]() { return PlayReverseClickedDelegate.Execute(*this); })
+				.ToolTipText(PlayReverseItemTooltipText)
+				[
+					SNew(SImage)
+					.Image(FSequencerPlaylistsStyle::Get().GetBrush("SequencerPlaylists.PlayReverse.Small"))
+					.ColorAndOpacity(FSlateColor::UseForeground())
+				]
+			]
+			+ SVerticalBox::Slot()
+			.Padding(8.0f, 2.0f, 0.0f, 2.0f)
 			[
 				SNew(SButton)
 				.ButtonStyle(FSequencerPlaylistsStyle::Get(), "SequencerPlaylists.HoverTransport.Stop")
@@ -1174,7 +1280,7 @@ TSharedRef<SWidget> SSequencerPlaylistItemWidget::GenerateWidgetForColumn(const
 				]
 			]
 			+ SVerticalBox::Slot()
-			.Padding(8.0f, 0.0f, 0.0f, 4.0f)
+			.Padding(8.0f, 2.0f, 0.0f, 4.0f)
 			[
 				SNew(SButton)
 				.ButtonStyle(FSequencerPlaylistsStyle::Get(), "SequencerPlaylists.HoverTransport.Reset")
@@ -1212,6 +1318,7 @@ TSharedRef<SWidget> SSequencerPlaylistItemWidget::GenerateWidgetForColumn(const
 	{
 		return SNew(SVerticalBox)
 			+ SVerticalBox::Slot()
+			.Padding(0.0f, 4.0f, 0.0f, 0.0f)
 			[
 				SNew(SNumericEntryBox<int32>)
 				.AllowSpin(true)
@@ -1233,6 +1340,7 @@ TSharedRef<SWidget> SSequencerPlaylistItemWidget::GenerateWidgetForColumn(const
 				})
 			]
 			+ SVerticalBox::Slot()
+			.Padding(0.0f, 0.0f, 0.0f, 4.0f)
 			[
 				SNew(SNumericEntryBox<int32>)
 				.AllowSpin(true)
@@ -1255,31 +1363,27 @@ TSharedRef<SWidget> SSequencerPlaylistItemWidget::GenerateWidgetForColumn(const
 			];
 	}
 
-	if (ColumnName == SSequencerPlaylistPanel::ColumnName_Hold)
+	if (ColumnName == SSequencerPlaylistPanel::ColumnName_Pause)
 	{
-		TSharedRef<SWidget> HoldToggle = SNew(SCheckBox)
+		TSharedRef<SWidget> PauseToggle = SNew(SCheckBox)
 			.Padding(FMargin(4.0f, 2.0f))
 			.HAlign(HAlign_Center)
-			.ToolTipText(LOCTEXT("ToggleHoldTooltip", "Enable or disable hold. Hold will infinitely hold the first frame of this item until manually played. Items are put into a hold state at the start of a take, or manually by hitting \"Reset.\""))
+			.ToolTipText(LOCTEXT("TogglePauseTooltip", "Enable or disable hold. Hold will infinitely hold the first frame of this item until manually played. Items are put into a hold state at the start of a take, or manually by hitting \"Reset.\""))
 			.Style(FAppStyle::Get(), "ToggleButtonCheckbox")
-			.IsChecked_Lambda([WeakItem]() { return (WeakItem.IsValid() && WeakItem->bHoldAtFirstFrame) ? ECheckBoxState::Checked : ECheckBoxState::Unchecked; })
-			.OnCheckStateChanged_Lambda([WeakItem](ECheckBoxState InState) {
-				if (WeakItem.IsValid())
-				{
-					WeakItem->bHoldAtFirstFrame = (InState == ECheckBoxState::Checked);
-				}
-			})
+			.IsChecked_Lambda([this]() { return IsPaused.Get() ? ECheckBoxState::Checked : ECheckBoxState::Unchecked; })
+			.OnCheckStateChanged_Lambda([this](ECheckBoxState InState) { PauseClickedDelegate.Execute(*this); })
 			[
 				SNew(SImage)
-				.Image(FSequencerPlaylistsStyle::Get().GetBrush("SequencerPlaylists.HoldFrame"))
-				.ColorAndOpacity(FSlateColor::UseForeground())
+				.Image(FSequencerPlaylistsStyle::Get().GetBrush("SequencerPlaylists.Pause"))
+				//.ColorAndOpacity(FStyleColors::AccentBlue)
 			];
 
 		return SNew(SVerticalBox)
 			+ SVerticalBox::Slot()
+			.Padding(0.0f, 1.0f, 0.0f, 0.0f)
 			.VAlign(VAlign_Center)
 			[
-				HoldToggle
+				PauseToggle
 			]
 			+ SVerticalBox::Slot()
 			[
@@ -1344,6 +1448,7 @@ TSharedRef<SWidget> SSequencerPlaylistItemWidget::GenerateWidgetForColumn(const
 
 		return SNew(SVerticalBox)
 			+ SVerticalBox::Slot()
+			.Padding(0.0f, 2.0f, 0.0f, 0.0f)
 			[
 				SNew(SHorizontalBox)
 				+ SHorizontalBox::Slot()
diff --git a/Engine/Plugins/VirtualProduction/SequencerPlaylists/Source/SequencerPlaylists/Private/SequencerPlaylistsWidgets.h b/Engine/Plugins/VirtualProduction/SequencerPlaylists/Source/SequencerPlaylists/Private/SequencerPlaylistsWidgets.h
--- a/Engine/Plugins/VirtualProduction/SequencerPlaylists/Source/SequencerPlaylists/Private/SequencerPlaylistsWidgets.h
+++ b/Engine/Plugins/VirtualProduction/SequencerPlaylists/Source/SequencerPlaylists/Private/SequencerPlaylistsWidgets.h
@@ -49,7 +49,7 @@ public:
 	static const FName ColumnName_HoverTransport;
 	static const FName ColumnName_Items;
 	static const FName ColumnName_Offset;
-	static const FName ColumnName_Hold;
+	static const FName ColumnName_Pause;
 	static const FName ColumnName_Loop;
 	static const FName ColumnName_HoverDetails;
 
@@ -90,11 +90,14 @@ private:
 	void OnSearchTextChanged(const FText& InFilterText);
 
 	FReply HandleClicked_PlayAll();
+	FReply HandleClicked_PlayAllReverse();
 	FReply HandleClicked_StopAll();
 	FReply HandleClicked_ResetAll();
 	FReply HandleClicked_AddSequence();
 
 	FReply HandleClicked_Item_Play(SSequencerPlaylistItemWidget& ItemWidget);
+	FReply HandleClicked_Item_PlayReverse(SSequencerPlaylistItemWidget& ItemWidget);
+	FReply HandleClicked_Item_Pause(SSequencerPlaylistItemWidget& ItemWidget);
 	FReply HandleClicked_Item_Stop(SSequencerPlaylistItemWidget& ItemWidget);
 	FReply HandleClicked_Item_Reset(SSequencerPlaylistItemWidget& ItemWidget);
 	FReply HandleClicked_Item_Remove(SSequencerPlaylistItemWidget& ItemWidget);
@@ -146,15 +149,20 @@ DECLARE_DELEGATE_RetVal_OneParam(FReply, FOnClickedSequencerPlaylistItem, SSeque
 
 class SSequencerPlaylistItemWidget : public SMultiColumnTableRow<TSharedPtr<FSequencerPlaylistRowData>>
 {
+	static const FText PlayReverseItemTooltipText;
 	static const FText PlayItemTooltipText;
+	static const FText PauseItemTooltipText;
 	static const FText StopItemTooltipText;
 	static const FText ResetItemTooltipText;
 
 	SLATE_BEGIN_ARGS(SSequencerPlaylistItemWidget) {}
 		SLATE_ATTRIBUTE(bool, PlayMode)
 		SLATE_ATTRIBUTE(bool, IsPlaying)
+		SLATE_ATTRIBUTE(bool, IsPaused)
 
 		SLATE_EVENT(FOnClickedSequencerPlaylistItem, OnPlayClicked)
+		SLATE_EVENT(FOnClickedSequencerPlaylistItem, OnPlayReverseClicked)
+		SLATE_EVENT(FOnClickedSequencerPlaylistItem, OnPauseClicked)
 		SLATE_EVENT(FOnClickedSequencerPlaylistItem, OnStopClicked)
 		SLATE_EVENT(FOnClickedSequencerPlaylistItem, OnResetClicked)
 		SLATE_EVENT(FOnClickedSequencerPlaylistItem, OnRemoveClicked)
@@ -213,8 +221,11 @@ private:
 
 	TAttribute<bool> PlayMode;
 	TAttribute<bool> IsPlaying;
+	TAttribute<bool> IsPaused;
 
 	FOnClickedSequencerPlaylistItem PlayClickedDelegate;
+	FOnClickedSequencerPlaylistItem PlayReverseClickedDelegate;
+	FOnClickedSequencerPlaylistItem PauseClickedDelegate;
 	FOnClickedSequencerPlaylistItem StopClickedDelegate;
 	FOnClickedSequencerPlaylistItem ResetClickedDelegate;
 	FOnClickedSequencerPlaylistItem RemoveClickedDelegate;
diff --git a/Engine/Plugins/VirtualProduction/SequencerPlaylists/Source/SequencerPlaylists/Public/ISequencerPlaylistsModule.h b/Engine/Plugins/VirtualProduction/SequencerPlaylists/Source/SequencerPlaylists/Public/ISequencerPlaylistsModule.h
--- a/Engine/Plugins/VirtualProduction/SequencerPlaylists/Source/SequencerPlaylists/Public/ISequencerPlaylistsModule.h
+++ b/Engine/Plugins/VirtualProduction/SequencerPlaylists/Source/SequencerPlaylists/Public/ISequencerPlaylistsModule.h
@@ -6,6 +6,7 @@
 #include "Delegates/Delegate.h"
 #include "Modules/ModuleManager.h"
 #include "Templates/SubclassOf.h"
+#include "ISequencerPlaylistsModule.generated.h"
 
 
 class ISequencer;
@@ -13,6 +14,30 @@ class USequencerPlaylistItem;
 class USequencerPlaylistPlayer;
 
 
+UENUM(BlueprintType)
+enum class ESequencerPlaylistPlaybackDirection : uint8
+{
+	Forward,
+	Reverse,
+};
+
+
+USTRUCT(BlueprintType)
+struct FSequencerPlaylistPlaybackState
+{
+	GENERATED_BODY()
+
+	UPROPERTY(BlueprintReadWrite, Category="Sequencer Playlists")
+	uint8 bIsPlaying : 1;
+
+	UPROPERTY(BlueprintReadWrite, Category="Sequencer Playlists")
+	uint8 bIsPaused : 1;
+
+	UPROPERTY(BlueprintReadWrite, Category="Sequencer Playlists")
+	ESequencerPlaylistPlaybackDirection PlaybackDirection;
+};
+
+
 /** Interface for derived classes that play one or more subclasses of USequencerPlaylistItem. */
 class ISequencerPlaylistItemPlayer
 {
@@ -23,7 +48,14 @@ public:
 	 * Initiate playback of the specified item.
 	 * @return True if the current sequence was modified, otherwise false.
 	 */
-	virtual bool Play(USequencerPlaylistItem* Item) = 0;
+	virtual bool Play(USequencerPlaylistItem* Item,
+	                  ESequencerPlaylistPlaybackDirection Direction = ESequencerPlaylistPlaybackDirection::Forward) = 0;
+
+	/**
+	 * Suspend playback of the specified item and begin a hold.
+	 * @return True if the current sequence was modified, otherwise false.
+	 */
+	virtual bool Pause(USequencerPlaylistItem* Item) = 0;
 
 	/**
 	 * Halt any current playback of the specified item.
@@ -44,8 +76,14 @@ public:
 	 */
 	virtual bool Reset(USequencerPlaylistItem* Item) = 0;
 
+	virtual FSequencerPlaylistPlaybackState GetPlaybackState(USequencerPlaylistItem* Item) const = 0;
+
 	/** Returns whether we're currently playing the specified item. */
-	virtual bool IsPlaying(USequencerPlaylistItem* Item) const = 0;
+	UE_DEPRECATED(5.2, "Use GetPlaybackState(Item).bIsPlaying instead.")
+	virtual bool IsPlaying(USequencerPlaylistItem* Item) const
+	{
+		return GetPlaybackState(Item).bIsPlaying;
+	}
 };
 
 
diff --git a/Engine/Plugins/VirtualProduction/SequencerPlaylists/Source/SequencerPlaylists/Public/SequencerPlaylistItem.h b/Engine/Plugins/VirtualProduction/SequencerPlaylists/Source/SequencerPlaylists/Public/SequencerPlaylistItem.h
--- a/Engine/Plugins/VirtualProduction/SequencerPlaylists/Source/SequencerPlaylists/Public/SequencerPlaylistItem.h
+++ b/Engine/Plugins/VirtualProduction/SequencerPlaylists/Source/SequencerPlaylists/Public/SequencerPlaylistItem.h
@@ -16,29 +16,32 @@ public:
 
 public:
 	/** Number of frames by which to clip the in point of sections played from this item. Will also affect the first frame for hold. */
-	UPROPERTY(EditAnywhere, BlueprintReadWrite, Category="SequencerPlaylists")
+	UPROPERTY(EditAnywhere, BlueprintReadWrite, Category="Sequencer Playlists")
 	int32 StartFrameOffset;
 
 	/** Number of frames by which to clip the out point of sections played from this item. */
-	UPROPERTY(EditAnywhere, BlueprintReadWrite, Category="SequencerPlaylists")
+	UPROPERTY(EditAnywhere, BlueprintReadWrite, Category="Sequencer Playlists")
 	int32 EndFrameOffset;
 
+#if WITH_EDITORONLY_DATA
 	/**
 	 * If true, the sequence will be inserted immediately on recording start and any time Reset()
 	 * is called, paused at the first frame indefinitely until subsequently played or stopped.
 	 */
-	UPROPERTY(EditAnywhere, BlueprintReadWrite, Category="SequencerPlaylists")
-	bool bHoldAtFirstFrame;
+	UE_DEPRECATED(5.2, "Holds are now created by invoking Pause prior to beginning recording.")
+	UPROPERTY(meta=(DeprecatedProperty, DeprecationMessage="Holds are now created by invoking Pause prior to beginning recording."))
+	bool bHoldAtFirstFrame_DEPRECATED;
+#endif
 
 	/** 0 is single playthrough, >= 1 is (n+1) playthroughs. */
-	UPROPERTY(EditAnywhere, BlueprintReadWrite, Category="SequencerPlaylists")
+	UPROPERTY(EditAnywhere, BlueprintReadWrite, Category="Sequencer Playlists")
 	int32 NumLoops;
 
 	/** Speed multiplier at which to play sections created by this item. A value of 1 will result in playback at the original speed. */
-	UPROPERTY(EditAnywhere, BlueprintReadWrite, Category="SequencerPlaylists")
+	UPROPERTY(EditAnywhere, BlueprintReadWrite, Category="Sequencer Playlists")
 	float PlaybackSpeed = 1.0f;
 
 	/** Disable playback of this item and prevent it from entering a hold state. */
-	UPROPERTY(EditAnywhere, BlueprintReadWrite, Category="SequencerPlaylists")
+	UPROPERTY(EditAnywhere, BlueprintReadWrite, Category="Sequencer Playlists")
 	bool bMute;
 };
diff --git a/Engine/Plugins/VirtualProduction/SequencerPlaylists/Source/SequencerPlaylists/Public/SequencerPlaylistPlayer.h b/Engine/Plugins/VirtualProduction/SequencerPlaylists/Source/SequencerPlaylists/Public/SequencerPlaylistPlayer.h
--- a/Engine/Plugins/VirtualProduction/SequencerPlaylists/Source/SequencerPlaylists/Public/SequencerPlaylistPlayer.h
+++ b/Engine/Plugins/VirtualProduction/SequencerPlaylists/Source/SequencerPlaylists/Public/SequencerPlaylistPlayer.h
@@ -3,6 +3,7 @@
 #pragma once
 
 #include "CoreMinimal.h"
+#include "ISequencerPlaylistsModule.h"
 #include "SequencerPlaylistItem.h"
 #include "Templates/SubclassOf.h"
 #include "Templates/UniquePtr.h"
@@ -35,39 +36,47 @@ public:
 	void BeginDestroy() override;
 	//~ End UObject interface
 
-	UFUNCTION(BlueprintCallable, Category="SequencerPlaylists")
+	UFUNCTION(BlueprintCallable, Category="Sequencer Playlists")
 	void SetPlaylist(USequencerPlaylist* InPlaylist);
 
-	UFUNCTION(BlueprintPure, Category="SequencerPlaylists")
+	UFUNCTION(BlueprintPure, Category="Sequencer Playlists")
 	USequencerPlaylist* GetPlaylist() { return Playlist; }
 
 	UE_DEPRECATED(5.1, "There is no longer a \"default\" player. Open a specific Playlist asset to create a player associated with it.")
 	UFUNCTION(BlueprintPure, meta=(DeprecatedFunction, DeprecationMessage= "There is no longer a \"default\" player. Open a specific Playlist asset to create a player associated with it.", DisplayName="Get Default Sequencer Playlist Player"), Category="SequencerPlaylists")
 	static USequencerPlaylistPlayer* GetDefaultPlayer() { return nullptr; }
 
-	UPROPERTY(BlueprintAssignable, Category="SequencerPlaylists")
+	UPROPERTY(BlueprintAssignable, Category="Sequencer Playlists")
 	FOnPlayerSequencerPlaylistSet OnPlaylistSet;
 
 public:
-	UFUNCTION(BlueprintCallable, Category="SequencerPlaylists")
-	bool PlayItem(USequencerPlaylistItem* Item);
+	UFUNCTION(BlueprintCallable, Category="Sequencer Playlists")
+	bool PlayItem(USequencerPlaylistItem* Item,
+	              ESequencerPlaylistPlaybackDirection Direction = ESequencerPlaylistPlaybackDirection::Forward);
 
-	UFUNCTION(BlueprintCallable, Category="SequencerPlaylists")
+	UFUNCTION(BlueprintCallable, Category="Sequencer Playlists")
+	bool PauseItem(USequencerPlaylistItem* Item);
+
+	UFUNCTION(BlueprintCallable, Category="Sequencer Playlists")
 	bool StopItem(USequencerPlaylistItem* Item);
 
-	UFUNCTION(BlueprintCallable, Category="SequencerPlaylists")
+	UFUNCTION(BlueprintCallable, Category="Sequencer Playlists")
 	bool ResetItem(USequencerPlaylistItem* Item);
 
-	UFUNCTION(BlueprintPure, Category="SequencerPlaylists")
+	UE_DEPRECATED(5.2, "Call GetPlaybackState and use the bIsPlaying flag in FSequencerPlaylistPlaybackState instead.")
+	UFUNCTION(BlueprintPure, Category="Sequencer Playlists")
 	bool IsPlaying(USequencerPlaylistItem* Item);
 
-	UFUNCTION(BlueprintCallable, Category="SequencerPlaylists")
-	bool PlayAll();
+	UFUNCTION(BlueprintPure, Category="Sequencer Playlists")
+	FSequencerPlaylistPlaybackState GetPlaybackState(USequencerPlaylistItem* Item);
+
+	UFUNCTION(BlueprintCallable, Category="Sequencer Playlists")
+	bool PlayAll(ESequencerPlaylistPlaybackDirection Direction = ESequencerPlaylistPlaybackDirection::Forward);
 
-	UFUNCTION(BlueprintCallable, Category="SequencerPlaylists")
+	UFUNCTION(BlueprintCallable, Category="Sequencer Playlists")
 	bool StopAll();
 
-	UFUNCTION(BlueprintCallable, Category="SequencerPlaylists")
+	UFUNCTION(BlueprintCallable, Category="Sequencer Playlists")
 	bool ResetAll();
 
 private:
