Sequencer Playlists: Add default constructor for FSequencerPlaylistPlaybackState, add missing "virtual" to some methods marked "override."

diff --git a/Engine/Plugins/VirtualProduction/SequencerPlaylists/Source/SequencerPlaylists/Private/SequencerPlaylistItem_Sequence.h b/Engine/Plugins/VirtualProduction/SequencerPlaylists/Source/SequencerPlaylists/Private/SequencerPlaylistItem_Sequence.h
--- a/Engine/Plugins/VirtualProduction/SequencerPlaylists/Source/SequencerPlaylists/Private/SequencerPlaylistItem_Sequence.h
+++ b/Engine/Plugins/VirtualProduction/SequencerPlaylists/Source/SequencerPlaylists/Private/SequencerPlaylistItem_Sequence.h
@@ -45,14 +45,14 @@ public:
 	~FSequencerPlaylistItemPlayer_Sequence() override;
 
 	//~ Begin ISequencerPlaylistItemPlayer
-	bool Play(USequencerPlaylistItem* Item,
-	          ESequencerPlaylistPlaybackDirection Direction = ESequencerPlaylistPlaybackDirection::Forward) override;
-	bool Pause(USequencerPlaylistItem* Item) override;
-	bool Stop(USequencerPlaylistItem* Item) override;
-	bool AddHold(USequencerPlaylistItem* Item) override;
-	bool Reset(USequencerPlaylistItem* Item) override;
-
-	FSequencerPlaylistPlaybackState GetPlaybackState(USequencerPlaylistItem* Item) const override;
+	virtual bool Play(USequencerPlaylistItem* Item,
+	                  ESequencerPlaylistPlaybackDirection Direction = ESequencerPlaylistPlaybackDirection::Forward) override;
+	virtual bool Pause(USequencerPlaylistItem* Item) override;
+	virtual bool Stop(USequencerPlaylistItem* Item) override;
+	virtual bool AddHold(USequencerPlaylistItem* Item) override;
+	virtual bool Reset(USequencerPlaylistItem* Item) override;
+
+	virtual FSequencerPlaylistPlaybackState GetPlaybackState(USequencerPlaylistItem* Item) const override;
 	//~ End ISequencerPlaylistItemPlayer
 
 private:
diff --git a/Engine/Plugins/VirtualProduction/SequencerPlaylists/Source/SequencerPlaylists/Public/ISequencerPlaylistsModule.h b/Engine/Plugins/VirtualProduction/SequencerPlaylists/Source/SequencerPlaylists/Public/ISequencerPlaylistsModule.h
--- a/Engine/Plugins/VirtualProduction/SequencerPlaylists/Source/SequencerPlaylists/Public/ISequencerPlaylistsModule.h
+++ b/Engine/Plugins/VirtualProduction/SequencerPlaylists/Source/SequencerPlaylists/Public/ISequencerPlaylistsModule.h
@@ -35,6 +35,13 @@ struct FSequencerPlaylistPlaybackState
 
 	UPROPERTY(BlueprintReadWrite, Category="Sequencer Playlists")
 	ESequencerPlaylistPlaybackDirection PlaybackDirection;
+
+	FSequencerPlaylistPlaybackState()
+		: bIsPlaying(false)
+		, bIsPaused(false)
+		, PlaybackDirection(ESequencerPlaylistPlaybackDirection::Forward)
+	{
+	}
 };
 
 
