Sequencer Playlists: Fix duplicate localization key.

diff --git a/Engine/Plugins/VirtualProduction/SequencerPlaylists/Source/SequencerPlaylists/Private/SequencerPlaylistsWidgets.cpp b/Engine/Plugins/VirtualProduction/SequencerPlaylists/Source/SequencerPlaylists/Private/SequencerPlaylistsWidgets.cpp
--- a/Engine/Plugins/VirtualProduction/SequencerPlaylists/Source/SequencerPlaylists/Private/SequencerPlaylistsWidgets.cpp
+++ b/Engine/Plugins/VirtualProduction/SequencerPlaylists/Source/SequencerPlaylists/Private/SequencerPlaylistsWidgets.cpp
@@ -345,7 +345,7 @@ TSharedRef<SWidget> SSequencerPlaylistPanel::Construct_Transport()
 			.HAlign(HAlign_Center)
 			.VAlign(VAlign_Center)
 			.OnClicked(this, &SSequencerPlaylistPanel::HandleClicked_PlayAll)
-			.ToolTipText(LOCTEXT("PlayAllReverseButtonTooltip", "Pause all items."))
+			.ToolTipText(LOCTEXT("PauseAllButtonTooltip", "Pause all items."))
 			[
 				SNew(SImage)
 				.Image(FSequencerPlaylistsStyle::Get().GetBrush("SequencerPlaylists.Pause"))
