Sequencer Playlists: Fix "Pause All" button.

Also separates creation of a pending take into its own transaction, as a partial mitigation for the fact that undoing its creation currently leads to a crash.

diff --git a/Engine/Plugins/VirtualProduction/SequencerPlaylists/Source/SequencerPlaylists/Private/SequencerPlaylistItem_Sequence.cpp b/Engine/Plugins/VirtualProduction/SequencerPlaylists/Source/SequencerPlaylists/Private/SequencerPlaylistItem_Sequence.cpp
--- a/Engine/Plugins/VirtualProduction/SequencerPlaylists/Source/SequencerPlaylists/Private/SequencerPlaylistItem_Sequence.cpp
+++ b/Engine/Plugins/VirtualProduction/SequencerPlaylists/Source/SequencerPlaylists/Private/SequencerPlaylistItem_Sequence.cpp
@@ -99,7 +99,7 @@ bool FSequencerPlaylistItemPlayer_Sequence::Play(
 }
 
 
-bool FSequencerPlaylistItemPlayer_Sequence::Pause(USequencerPlaylistItem* Item)
+bool FSequencerPlaylistItemPlayer_Sequence::TogglePause(USequencerPlaylistItem* Item)
 {
 	return InternalPause(Item);
 }
diff --git a/Engine/Plugins/VirtualProduction/SequencerPlaylists/Source/SequencerPlaylists/Private/SequencerPlaylistItem_Sequence.h b/Engine/Plugins/VirtualProduction/SequencerPlaylists/Source/SequencerPlaylists/Private/SequencerPlaylistItem_Sequence.h
--- a/Engine/Plugins/VirtualProduction/SequencerPlaylists/Source/SequencerPlaylists/Private/SequencerPlaylistItem_Sequence.h
+++ b/Engine/Plugins/VirtualProduction/SequencerPlaylists/Source/SequencerPlaylists/Private/SequencerPlaylistItem_Sequence.h
@@ -47,7 +47,7 @@ public:
 	//~ Begin ISequencerPlaylistItemPlayer
 	virtual bool Play(USequencerPlaylistItem* Item,
 	                  ESequencerPlaylistPlaybackDirection Direction = ESequencerPlaylistPlaybackDirection::Forward) override;
-	virtual bool Pause(USequencerPlaylistItem* Item) override;
+	virtual bool TogglePause(USequencerPlaylistItem* Item) override;
 	virtual bool Stop(USequencerPlaylistItem* Item) override;
 	virtual bool AddHold(USequencerPlaylistItem* Item) override;
 	virtual bool Reset(USequencerPlaylistItem* Item) override;
diff --git a/Engine/Plugins/VirtualProduction/SequencerPlaylists/Source/SequencerPlaylists/Private/SequencerPlaylistPlayer.cpp b/Engine/Plugins/VirtualProduction/SequencerPlaylists/Source/SequencerPlaylists/Private/SequencerPlaylistPlayer.cpp
--- a/Engine/Plugins/VirtualProduction/SequencerPlaylists/Source/SequencerPlaylists/Private/SequencerPlaylistPlayer.cpp
+++ b/Engine/Plugins/VirtualProduction/SequencerPlaylists/Source/SequencerPlaylists/Private/SequencerPlaylistPlayer.cpp
@@ -70,13 +70,15 @@ bool USequencerPlaylistPlayer::PlayItem(
 		return false;
 	}
 
+	TSharedPtr<ISequencer> Sequencer = GetOrCreateSequencer();
+
 	FScopedTransaction Transaction(FText::Format(LOCTEXT("PlayItemTransaction", "Begin playback of {0}"), Item->GetDisplayName()));
 
 	EnterUnboundedPlayIfNotRecording();
 
 	if (GetCheckedItemPlayer(Item)->Play(Item, Direction))
 	{
-		GetSequencer()->NotifyMovieSceneDataChanged(EMovieSceneDataChangeType::RefreshAllImmediately);
+		Sequencer->NotifyMovieSceneDataChanged(EMovieSceneDataChangeType::RefreshAllImmediately);
 		return true;
 	}
 	else
@@ -93,11 +95,13 @@ bool USequencerPlaylistPlayer::PauseItem(USequencerPlaylistItem* Item)
 		return false;
 	}
 
+	TSharedPtr<ISequencer> Sequencer = GetOrCreateSequencer();
+
 	FScopedTransaction Transaction(FText::Format(LOCTEXT("PauseItemTransaction", "Toggle pause of {0}"), Item->GetDisplayName()));
 
-	if (GetCheckedItemPlayer(Item)->Pause(Item))
+	if (GetCheckedItemPlayer(Item)->TogglePause(Item))
 	{
-		GetSequencer()->NotifyMovieSceneDataChanged(EMovieSceneDataChangeType::RefreshAllImmediately);
+		Sequencer->NotifyMovieSceneDataChanged(EMovieSceneDataChangeType::RefreshAllImmediately);
 		return true;
 	}
 	else
@@ -114,11 +118,13 @@ bool USequencerPlaylistPlayer::StopItem(USequencerPlaylistItem* Item)
 		return false;
 	}
 
+	TSharedPtr<ISequencer> Sequencer = GetOrCreateSequencer();
+
 	FScopedTransaction Transaction(FText::Format(LOCTEXT("StopItemTransaction", "Stop playback of {0}"), Item->GetDisplayName()));
 
 	if (GetCheckedItemPlayer(Item)->Stop(Item))
 	{
-		GetSequencer()->NotifyMovieSceneDataChanged(EMovieSceneDataChangeType::RefreshAllImmediately);
+		Sequencer->NotifyMovieSceneDataChanged(EMovieSceneDataChangeType::RefreshAllImmediately);
 		return true;
 	}
 	else
@@ -135,11 +141,13 @@ bool USequencerPlaylistPlayer::ResetItem(USequencerPlaylistItem* Item)
 		return false;
 	}
 
+	TSharedPtr<ISequencer> Sequencer = GetOrCreateSequencer();
+
 	FScopedTransaction Transaction(FText::Format(LOCTEXT("ResetItemTransaction", "Reset playback of {0}"), Item->GetDisplayName()));
 
 	if (GetCheckedItemPlayer(Item)->Reset(Item))
 	{
-		GetSequencer()->NotifyMovieSceneDataChanged(EMovieSceneDataChangeType::RefreshAllImmediately);
+		Sequencer->NotifyMovieSceneDataChanged(EMovieSceneDataChangeType::RefreshAllImmediately);
 		return true;
 	}
 	else
@@ -271,7 +279,7 @@ void USequencerPlaylistPlayer::Tick(float DeltaTime)
 
 void USequencerPlaylistPlayer::EnterUnboundedPlayIfNotRecording()
 {
-	TSharedPtr<ISequencer> Sequencer = GetSequencer();
+	TSharedPtr<ISequencer> Sequencer = GetOrCreateSequencer();
 	UTakeRecorder* Recorder = UTakeRecorder::GetActiveRecorder();
 
 	const bool bInRecorder = Recorder && Recorder->GetState() != ETakeRecorderState::Stopped;
@@ -299,6 +307,8 @@ bool USequencerPlaylistPlayer::PlayAll(
 		return false;
 	}
 
+	TSharedPtr<ISequencer> Sequencer = GetOrCreateSequencer();
+
 	FScopedTransaction Transaction(LOCTEXT("PlayAllTransaction", "Begin playback of all Playlist items"));
 
 	EnterUnboundedPlayIfNotRecording();
@@ -312,12 +322,67 @@ bool USequencerPlaylistPlayer::PlayAll(
 			continue;
 		}
 
-		bAnyChange |= GetCheckedItemPlayer(Item)->Play(Item, Direction);
+		TSharedPtr<ISequencerPlaylistItemPlayer> Player = GetCheckedItemPlayer(Item);
+		const FSequencerPlaylistPlaybackState ItemState = Player->GetPlaybackState(Item);
+
+		if (ItemState.bIsPaused)
+		{
+			bAnyChange |= Player->TogglePause(Item);
+		}
+		else
+		{
+			bAnyChange |= Player->Play(Item, Direction);
+		}
 	}
 
 	if (bAnyChange)
 	{
-		GetSequencer()->NotifyMovieSceneDataChanged(EMovieSceneDataChangeType::RefreshAllImmediately);
+		Sequencer->NotifyMovieSceneDataChanged(EMovieSceneDataChangeType::RefreshAllImmediately);
+	}
+
+	return bAnyChange;
+}
+
+
+bool USequencerPlaylistPlayer::PauseAll()
+{
+	if (!ensure(Playlist) || !Playlist->Items.Num())
+	{
+		return false;
+	}
+
+	TSharedPtr<ISequencer> Sequencer = GetOrCreateSequencer();
+	UTakeRecorder* Recorder = UTakeRecorder::GetActiveRecorder();
+
+	const bool bPlaying = Sequencer->GetPlaybackStatus() == EMovieScenePlayerStatus::Playing;
+	const bool bRecording = Recorder && Recorder->GetState() != ETakeRecorderState::Stopped;
+
+	bool bAnyChange = false;
+
+	FScopedTransaction Transaction(LOCTEXT("PauseAllTransaction", "Pause playback of all Playlist items"));
+
+	// If we are playing or recording, pause any playing items.
+	// If we are not playing or recording, add holds for any items without holds.
+	for (USequencerPlaylistItem* Item : Playlist->Items)
+	{
+		TSharedPtr<ISequencerPlaylistItemPlayer> Player = GetCheckedItemPlayer(Item);
+		const FSequencerPlaylistPlaybackState ItemState = Player->GetPlaybackState(Item);
+		if (bPlaying || bRecording)
+		{
+			if (ItemState.bIsPlaying && !ItemState.bIsPaused)
+			{
+				bAnyChange |= Player->TogglePause(Item);
+			}
+		}
+		else if (!ItemState.bIsPaused)
+		{
+			bAnyChange |= Player->AddHold(Item);
+		}
+	}
+
+	if (bAnyChange)
+	{
+		Sequencer->NotifyMovieSceneDataChanged(EMovieSceneDataChangeType::RefreshAllImmediately);
 	}
 
 	return bAnyChange;
@@ -339,7 +404,7 @@ bool USequencerPlaylistPlayer::StopAll()
 		bAnyChange |= GetCheckedItemPlayer(Item)->Stop(Item);
 	}
 
-	TSharedPtr<ISequencer> Sequencer = GetSequencer();
+	TSharedPtr<ISequencer> Sequencer = GetOrCreateSequencer();
 
 	UTakeRecorder* Recorder = UTakeRecorder::GetActiveRecorder();
 	const bool bInRecorder = Recorder && Recorder->GetState() != ETakeRecorderState::Stopped;
@@ -366,6 +431,8 @@ bool USequencerPlaylistPlayer::ResetAll()
 		return false;
 	}
 
+	TSharedPtr<ISequencer> Sequencer = GetOrCreateSequencer();
+
 	bool bAnyChange = false;
 
 	FScopedTransaction Transaction(LOCTEXT("ResetAllTransaction", "Reset playback of all Playlist items"));
@@ -384,14 +451,14 @@ bool USequencerPlaylistPlayer::ResetAll()
 
 	if (bAnyChange)
 	{
-		GetSequencer()->NotifyMovieSceneDataChanged(EMovieSceneDataChangeType::RefreshAllImmediately);
+		Sequencer->NotifyMovieSceneDataChanged(EMovieSceneDataChangeType::RefreshAllImmediately);
 	}
 
 	return bAnyChange;
 }
 
 
-TSharedPtr<ISequencer> USequencerPlaylistPlayer::GetSequencer()
+TSharedPtr<ISequencer> USequencerPlaylistPlayer::GetOrCreateSequencer()
 {
 	if (TSharedPtr<ISequencer> Sequencer = WeakSequencer.Pin())
 	{
@@ -432,33 +499,6 @@ TSharedPtr<ISequencer> USequencerPlaylistPlayer::GetSequencer()
 }
 
 
-TSharedPtr<ISequencer> USequencerPlaylistPlayer::GetValidatedSequencer()
-{
-	TSharedPtr<ISequencer> Sequencer = GetSequencer();
-	if (!Sequencer)
-	{
-		return nullptr;
-	}
-
-	ULevelSequence* RootSequence = Cast<ULevelSequence>(Sequencer->GetRootMovieSceneSequence());
-	if (!RootSequence)
-	{
-		UE_LOG(LogSequencerPlaylists, Error, TEXT("USequencerPlaylistPlayer::GetValidatedSequencer: Unable to get root sequence"));
-		return nullptr;
-	}
-
-	UMovieScene* RootScene = RootSequence->GetMovieScene();
-	if (!RootScene)
-	{
-		// TODO: Seems like this may not be possible?
-		UE_LOG(LogSequencerPlaylists, Error, TEXT("USequencerPlaylistPlayer::GetValidatedSequencer: Unable to get root scene"));
-		return nullptr;
-	}
-
-	return Sequencer;
-}
-
-
 void USequencerPlaylistPlayer::OnTakeRecorderInitialized(UTakeRecorder* InRecorder)
 {
 	if (InRecorder)
@@ -489,7 +529,7 @@ void USequencerPlaylistPlayer::OnTakeRecorderStopped(UTakeRecorder* InRecorder)
 	}
 
 	// FIXME: Any sequences not already stopped end up a few frames too long; pass in explicit end frame?
-	if (TSharedPtr<ISequencer> Sequencer = GetValidatedSequencer())
+	if (TSharedPtr<ISequencer> Sequencer = GetOrCreateSequencer())
 	{
 		FScopedTransaction Transaction(LOCTEXT("TakeRecorderStoppedTransaction", "Playlist - Take Recorder stopped"));
 		bool bAnySequencesModified = false;
@@ -519,7 +559,7 @@ TSharedPtr<ISequencerPlaylistItemPlayer> USequencerPlaylistPlayer::GetCheckedIte
 		return *ExistingPlayer;
 	}
 
-	TSharedPtr<ISequencer> Sequencer = GetValidatedSequencer();
+	TSharedPtr<ISequencer> Sequencer = GetOrCreateSequencer();
 	check(Sequencer.IsValid());
 
 	TSharedPtr<ISequencerPlaylistItemPlayer> NewPlayer =
diff --git a/Engine/Plugins/VirtualProduction/SequencerPlaylists/Source/SequencerPlaylists/Private/SequencerPlaylistsWidgets.cpp b/Engine/Plugins/VirtualProduction/SequencerPlaylists/Source/SequencerPlaylists/Private/SequencerPlaylistsWidgets.cpp
--- a/Engine/Plugins/VirtualProduction/SequencerPlaylists/Source/SequencerPlaylists/Private/SequencerPlaylistsWidgets.cpp
+++ b/Engine/Plugins/VirtualProduction/SequencerPlaylists/Source/SequencerPlaylists/Private/SequencerPlaylistsWidgets.cpp
@@ -344,7 +344,7 @@ TSharedRef<SWidget> SSequencerPlaylistPanel::Construct_Transport()
 			.ContentPadding(FMargin(0.0f, 2.0f))
 			.HAlign(HAlign_Center)
 			.VAlign(VAlign_Center)
-			.OnClicked(this, &SSequencerPlaylistPanel::HandleClicked_PlayAll)
+			.OnClicked(this, &SSequencerPlaylistPanel::HandleClicked_PauseAll)
 			.ToolTipText(LOCTEXT("PauseAllButtonTooltip", "Pause all items."))
 			[
 				SNew(SImage)
@@ -825,6 +825,19 @@ FReply SSequencerPlaylistPanel::HandleClicked_PlayAllReverse()
 }
 
 
+FReply SSequencerPlaylistPanel::HandleClicked_PauseAll()
+{
+	USequencerPlaylistPlayer* Player = WeakPlayer.Get();
+	if (!ensure(Player))
+	{
+		return FReply::Unhandled();
+	}
+
+	Player->PauseAll();
+	return FReply::Handled();
+}
+
+
 FReply SSequencerPlaylistPanel::HandleClicked_StopAll()
 {
 	USequencerPlaylistPlayer* Player = WeakPlayer.Get();
diff --git a/Engine/Plugins/VirtualProduction/SequencerPlaylists/Source/SequencerPlaylists/Private/SequencerPlaylistsWidgets.h b/Engine/Plugins/VirtualProduction/SequencerPlaylists/Source/SequencerPlaylists/Private/SequencerPlaylistsWidgets.h
--- a/Engine/Plugins/VirtualProduction/SequencerPlaylists/Source/SequencerPlaylists/Private/SequencerPlaylistsWidgets.h
+++ b/Engine/Plugins/VirtualProduction/SequencerPlaylists/Source/SequencerPlaylists/Private/SequencerPlaylistsWidgets.h
@@ -91,6 +91,7 @@ private:
 
 	FReply HandleClicked_PlayAll();
 	FReply HandleClicked_PlayAllReverse();
+	FReply HandleClicked_PauseAll();
 	FReply HandleClicked_StopAll();
 	FReply HandleClicked_ResetAll();
 	FReply HandleClicked_AddSequence();
diff --git a/Engine/Plugins/VirtualProduction/SequencerPlaylists/Source/SequencerPlaylists/Public/ISequencerPlaylistsModule.h b/Engine/Plugins/VirtualProduction/SequencerPlaylists/Source/SequencerPlaylists/Public/ISequencerPlaylistsModule.h
--- a/Engine/Plugins/VirtualProduction/SequencerPlaylists/Source/SequencerPlaylists/Public/ISequencerPlaylistsModule.h
+++ b/Engine/Plugins/VirtualProduction/SequencerPlaylists/Source/SequencerPlaylists/Public/ISequencerPlaylistsModule.h
@@ -59,10 +59,11 @@ public:
 	                  ESequencerPlaylistPlaybackDirection Direction = ESequencerPlaylistPlaybackDirection::Forward) = 0;
 
 	/**
-	 * Suspend playback of the specified item and begin a hold.
+	 * Suspend playback of the specified item (if it is playing) and begin a hold.
+	 * If it is already paused, resume playback.
 	 * @return True if the current sequence was modified, otherwise false.
 	 */
-	virtual bool Pause(USequencerPlaylistItem* Item) = 0;
+	virtual bool TogglePause(USequencerPlaylistItem* Item) = 0;
 
 	/**
 	 * Halt any current playback of the specified item.
diff --git a/Engine/Plugins/VirtualProduction/SequencerPlaylists/Source/SequencerPlaylists/Public/SequencerPlaylistPlayer.h b/Engine/Plugins/VirtualProduction/SequencerPlaylists/Source/SequencerPlaylists/Public/SequencerPlaylistPlayer.h
--- a/Engine/Plugins/VirtualProduction/SequencerPlaylists/Source/SequencerPlaylists/Public/SequencerPlaylistPlayer.h
+++ b/Engine/Plugins/VirtualProduction/SequencerPlaylists/Source/SequencerPlaylists/Public/SequencerPlaylistPlayer.h
@@ -73,6 +73,9 @@ public:
 	UFUNCTION(BlueprintCallable, Category="Sequencer Playlists")
 	bool PlayAll(ESequencerPlaylistPlaybackDirection Direction = ESequencerPlaylistPlaybackDirection::Forward);
 
+	UFUNCTION(BlueprintCallable, Category="Sequencer Playlists")
+	bool PauseAll();
+
 	UFUNCTION(BlueprintCallable, Category="Sequencer Playlists")
 	bool StopAll();
 
@@ -80,7 +83,7 @@ public:
 	bool ResetAll();
 
 private:
-	TSharedPtr<ISequencer> GetSequencer();
+	TSharedPtr<ISequencer> GetOrCreateSequencer();
 
 	/**
 	 * Enters unbounded play (like Take Recorder) if the user is not recording.
@@ -92,9 +95,6 @@ private:
 	 */
 	void Tick(float DeltaTime);
 
-	/** Centralizes additional common assumptions that are convenient for downstream code to make. */
-	TSharedPtr<ISequencer> GetValidatedSequencer();
-
 	void OnTakeRecorderInitialized(UTakeRecorder* InRecorder);
 	void OnTakeRecorderStarted(UTakeRecorder* InRecorder);
 	void OnTakeRecorderStopped(UTakeRecorder* InRecorder);
