Sequencer Playlists: Fix reverse and pause controls missing from Play Mode overlay.

diff --git a/Engine/Plugins/VirtualProduction/SequencerPlaylists/Source/SequencerPlaylists/Private/SequencerPlaylistsWidgets.cpp b/Engine/Plugins/VirtualProduction/SequencerPlaylists/Source/SequencerPlaylists/Private/SequencerPlaylistsWidgets.cpp
--- a/Engine/Plugins/VirtualProduction/SequencerPlaylists/Source/SequencerPlaylists/Private/SequencerPlaylistsWidgets.cpp
+++ b/Engine/Plugins/VirtualProduction/SequencerPlaylists/Source/SequencerPlaylists/Private/SequencerPlaylistsWidgets.cpp
@@ -1168,6 +1168,38 @@ void SSequencerPlaylistItemWidget::ConstructChildren(ETableViewMode::Type InOwne
 			.Visibility(this, &SSequencerPlaylistItemWidget::GetPlayModeTransportVisibility)
 			+ SHorizontalBox::Slot()
 			.AutoWidth()
+			[
+				SNew(SButton)
+				.ButtonStyle(FSequencerPlaylistsStyle::Get(), "SequencerPlaylists.TransportButton.PlayReverse")
+				.ContentPadding(6.0f)
+				.HAlign(HAlign_Center)
+				.VAlign(VAlign_Center)
+				.OnClicked_Lambda([this]() { return PlayReverseClickedDelegate.Execute(*this); })
+				.ToolTipText(PlayReverseItemTooltipText)
+				[
+					SNew(SImage)
+					.Image(FSequencerPlaylistsStyle::Get().GetBrush("SequencerPlaylists.PlayReverse"))
+					.ColorAndOpacity(FSlateColor::UseForeground())
+				]
+			]
+			+ SHorizontalBox::Slot()
+			.AutoWidth()
+			[
+				SNew(SButton)
+				.ButtonStyle(FSequencerPlaylistsStyle::Get(), "SequencerPlaylists.TransportButton.Pause")
+				.ContentPadding(6.0f)
+				.HAlign(HAlign_Center)
+				.VAlign(VAlign_Center)
+				.OnClicked_Lambda([this]() { return PauseClickedDelegate.Execute(*this); })
+				.ToolTipText(PauseItemTooltipText)
+				[
+					SNew(SImage)
+					.Image(FSequencerPlaylistsStyle::Get().GetBrush("SequencerPlaylists.Pause"))
+					.ColorAndOpacity(FSlateColor::UseForeground())
+				]
+			]
+			+ SHorizontalBox::Slot()
+			.AutoWidth()
 			[
 				SNew(SButton)
 				.ButtonStyle(FSequencerPlaylistsStyle::Get(), "SequencerPlaylists.TransportButton.Play")
