Sequencer Playlists: Fix hover transport column layout.

diff --git a/Engine/Plugins/VirtualProduction/SequencerPlaylists/Source/SequencerPlaylists/Private/SequencerPlaylistsStyle.cpp b/Engine/Plugins/VirtualProduction/SequencerPlaylists/Source/SequencerPlaylists/Private/SequencerPlaylistsStyle.cpp
--- a/Engine/Plugins/VirtualProduction/SequencerPlaylists/Source/SequencerPlaylists/Private/SequencerPlaylistsStyle.cpp
+++ b/Engine/Plugins/VirtualProduction/SequencerPlaylists/Source/SequencerPlaylists/Private/SequencerPlaylistsStyle.cpp
@@ -92,6 +92,12 @@ TSharedRef<FSlateStyleSet> FSequencerPlaylistsStyle::Create()
 	Style->Set("SequencerPlaylists.Loop.Disabled", new IMAGE_BRUSH_SVG("LoopNot", Icon16x16));
 	Style->Set("SequencerPlaylists.Loop.Finite", new IMAGE_BRUSH_SVG("LoopNumber", Icon16x16));
 
+	// Play, Stop, Reset are registered under a different root below.
+	Style->Set("SequencerPlaylists.Pause", new IMAGE_BRUSH_SVG("Pause", Icon20x20));
+	Style->Set("SequencerPlaylists.Pause.Small", new IMAGE_BRUSH_SVG("Pause", Icon16x16));
+	Style->Set("SequencerPlaylists.PlayReverse", new IMAGE_BRUSH_SVG("PlayReverse", Icon20x20));
+	Style->Set("SequencerPlaylists.PlayReverse.Small", new IMAGE_BRUSH_SVG("PlayReverse", Icon16x16));
+
 	// Main toolbar
 	{
 		const FToolBarStyle& SlimToolbarStyle = FAppStyle::Get().GetWidgetStyle<FToolBarStyle>("SlimToolBar");
@@ -278,15 +284,11 @@ TSharedRef<FSlateStyleSet> FSequencerPlaylistsStyle::Create()
 			Style->Set("SequencerPlaylists.SavePlaylist", new IMAGE_BRUSH_SVG("Starship/Common/SaveCurrent", Icon16x16));
 
 			Style->Set("SequencerPlaylists.Play", new IMAGE_BRUSH_SVG("Starship/Common/play", Icon20x20));
-			Style->Set("SequencerPlaylists.Play.Small", new IMAGE_BRUSH_SVG("Starship/Common/play", Icon12x12));
-			Style->Set("SequencerPlaylists.Pause", new IMAGE_BRUSH_SVG("Starship/MainToolbar/pause", Icon20x20)); // FIXME
-			Style->Set("SequencerPlaylists.Pause.Small", new IMAGE_BRUSH_SVG("Starship/MainToolbar/pause", Icon12x12)); // FIXME
-			Style->Set("SequencerPlaylists.PlayReverse", new IMAGE_BRUSH("Sequencer/Transport_Bar/Backwards_24x", Icon20x20)); // FIXME
-			Style->Set("SequencerPlaylists.PlayReverse.Small", new IMAGE_BRUSH("Sequencer/Transport_Bar/Backwards_24x", Icon12x12)); // FIXME
+			Style->Set("SequencerPlaylists.Play.Small", new IMAGE_BRUSH_SVG("Starship/Common/play", Icon16x16));
 			Style->Set("SequencerPlaylists.Stop", new IMAGE_BRUSH_SVG("Starship/Common/stop", Icon20x20));
-			Style->Set("SequencerPlaylists.Stop.Small", new IMAGE_BRUSH_SVG("Starship/Common/stop", Icon12x12));
+			Style->Set("SequencerPlaylists.Stop.Small", new IMAGE_BRUSH_SVG("Starship/Common/stop", Icon16x16));
 			Style->Set("SequencerPlaylists.Reset", new IMAGE_BRUSH_SVG("Starship/Common/Reset", Icon20x20));
-			Style->Set("SequencerPlaylists.Reset.Small", new IMAGE_BRUSH_SVG("Starship/Common/Reset", Icon12x12));
+			Style->Set("SequencerPlaylists.Reset.Small", new IMAGE_BRUSH_SVG("Starship/Common/Reset", Icon16x16));
 		}
 	}
 
diff --git a/Engine/Plugins/VirtualProduction/SequencerPlaylists/Source/SequencerPlaylists/Private/SequencerPlaylistsWidgets.cpp b/Engine/Plugins/VirtualProduction/SequencerPlaylists/Source/SequencerPlaylists/Private/SequencerPlaylistsWidgets.cpp
--- a/Engine/Plugins/VirtualProduction/SequencerPlaylists/Source/SequencerPlaylists/Private/SequencerPlaylistsWidgets.cpp
+++ b/Engine/Plugins/VirtualProduction/SequencerPlaylists/Source/SequencerPlaylists/Private/SequencerPlaylistsWidgets.cpp
@@ -474,7 +474,7 @@ TSharedRef<SWidget> SSequencerPlaylistPanel::Construct_ItemListView()
 			SNew(SHeaderRow)
 			+ SHeaderRow::Column(ColumnName_HoverTransport)
 				.DefaultLabel(FText::GetEmpty())
-				.FillSized(30.0f)
+				.FillSized(60.0f)
 			+ SHeaderRow::Column(ColumnName_Items)
 				.DefaultLabel(LOCTEXT("ColumnLabelItems", "Playlist Items"))
 				.FillWidth(1.0f)
@@ -483,7 +483,7 @@ TSharedRef<SWidget> SSequencerPlaylistPanel::Construct_ItemListView()
 				.FillSized(45.0f)
 			+ SHeaderRow::Column(ColumnName_Pause)
 				.DefaultLabel(LOCTEXT("ColumnLabelPause", "Pause"))
-				.FillSized(35.0f)
+				.FillSized(45.0f)
 				.HAlignCell(HAlign_Center)
 			+ SHeaderRow::Column(ColumnName_Loop)
 				.DefaultLabel(LOCTEXT("ColumnLabelLoop", "Loop"))
@@ -1273,72 +1273,82 @@ TSharedRef<SWidget> SSequencerPlaylistItemWidget::GenerateWidgetForColumn(const
 	{
 		return SNew(SVerticalBox)
 			+ SVerticalBox::Slot()
-			.Padding(8.0f, 4.0f, 0.0f, 2.0f)
+			.Padding(12.0f, 4.0f, 0.0f, 2.0f)
 			[
-				SNew(SButton)
-				.ButtonStyle(FSequencerPlaylistsStyle::Get(), "SequencerPlaylists.HoverTransport.Play")
-				.ContentPadding(0)
-				.HAlign(HAlign_Center)
-				.VAlign(VAlign_Center)
-				.Visibility_Lambda([this]() { return (IsHovered() && !InPlayMode()) ? EVisibility::Visible : EVisibility::Hidden; })
-				.OnClicked_Lambda([this]() { return PlayClickedDelegate.Execute(*this); })
-				.ToolTipText(PlayItemTooltipText)
+				SNew(SHorizontalBox)
+				+ SHorizontalBox::Slot()
+				.Padding(0.0f, 0.0f, 0.0f, 0.0f)
 				[
-					SNew(SImage)
-					.Image(FSequencerPlaylistsStyle::Get().GetBrush("SequencerPlaylists.Play.Small"))
-					.ColorAndOpacity(FSlateColor::UseForeground())
+					SNew(SButton)
+					.ButtonStyle(FSequencerPlaylistsStyle::Get(), "SequencerPlaylists.HoverTransport.PlayReverse")
+					.ContentPadding(0)
+					.HAlign(HAlign_Center)
+					.VAlign(VAlign_Center)
+					.Visibility_Lambda([this]() { return (IsHovered() && !InPlayMode()) ? EVisibility::Visible : EVisibility::Hidden; })
+					.OnClicked_Lambda([this]() { return PlayReverseClickedDelegate.Execute(*this); })
+					.ToolTipText(PlayReverseItemTooltipText)
+					[
+						SNew(SImage)
+						.Image(FSequencerPlaylistsStyle::Get().GetBrush("SequencerPlaylists.PlayReverse.Small"))
+						.ColorAndOpacity(FSlateColor::UseForeground())
+					]
 				]
-			]
-			+ SVerticalBox::Slot()
-			.Padding(8.0f, 2.0f, 0.0f, 2.0f)
-			[
-				SNew(SButton)
-				.ButtonStyle(FSequencerPlaylistsStyle::Get(), "SequencerPlaylists.HoverTransport.PlayReverse")
-				.ContentPadding(0)
-				.HAlign(HAlign_Center)
-				.VAlign(VAlign_Center)
-				.Visibility_Lambda([this]() { return (IsHovered() && !InPlayMode()) ? EVisibility::Visible : EVisibility::Hidden; })
-				.OnClicked_Lambda([this]() { return PlayReverseClickedDelegate.Execute(*this); })
-				.ToolTipText(PlayReverseItemTooltipText)
+				+ SHorizontalBox::Slot()
+				.Padding(2.0f, 0.0f, 0.0f, 0.0f)
 				[
-					SNew(SImage)
-					.Image(FSequencerPlaylistsStyle::Get().GetBrush("SequencerPlaylists.PlayReverse.Small"))
-					.ColorAndOpacity(FSlateColor::UseForeground())
+					SNew(SButton)
+					.ButtonStyle(FSequencerPlaylistsStyle::Get(), "SequencerPlaylists.HoverTransport.Play")
+					.ContentPadding(0)
+					.HAlign(HAlign_Center)
+					.VAlign(VAlign_Center)
+					.Visibility_Lambda([this]() { return (IsHovered() && !InPlayMode()) ? EVisibility::Visible : EVisibility::Hidden; })
+					.OnClicked_Lambda([this]() { return PlayClickedDelegate.Execute(*this); })
+					.ToolTipText(PlayItemTooltipText)
+					[
+						SNew(SImage)
+						.Image(FSequencerPlaylistsStyle::Get().GetBrush("SequencerPlaylists.Play.Small"))
+						.ColorAndOpacity(FSlateColor::UseForeground())
+					]
 				]
 			]
 			+ SVerticalBox::Slot()
-			.Padding(8.0f, 2.0f, 0.0f, 2.0f)
+			.Padding(12.0f, 2.0f, 0.0f, 4.0f)
 			[
-				SNew(SButton)
-				.ButtonStyle(FSequencerPlaylistsStyle::Get(), "SequencerPlaylists.HoverTransport.Stop")
-				.ContentPadding(0)
-				.HAlign(HAlign_Center)
-				.VAlign(VAlign_Center)
-				.Visibility_Lambda([this]() { return (IsPlaying.Get() || (IsHovered() && !InPlayMode())) ? EVisibility::Visible : EVisibility::Hidden; })
-				.ForegroundColor_Lambda([this]() { return IsPlaying.Get() ? FStyleColors::AccentRed : FSlateColor::UseStyle(); })
-				.OnClicked_Lambda([this]() { return StopClickedDelegate.Execute(*this); })
-				.ToolTipText(StopItemTooltipText)
+				SNew(SHorizontalBox)
+				+ SHorizontalBox::Slot()
+				.Padding(0.0f, 0.0f, 0.0f, 0.0f)
 				[
-					SNew(SImage)
-					.Image(FSequencerPlaylistsStyle::Get().GetBrush("SequencerPlaylists.Stop.Small"))
-					.ColorAndOpacity(FSlateColor::UseForeground())
+					SNew(SButton)
+					.ButtonStyle(FSequencerPlaylistsStyle::Get(), "SequencerPlaylists.HoverTransport.Stop")
+					.ContentPadding(0)
+					.HAlign(HAlign_Center)
+					.VAlign(VAlign_Center)
+					.Visibility_Lambda([this]() { return (IsPlaying.Get() || (IsHovered() && !InPlayMode())) ? EVisibility::Visible : EVisibility::Hidden; })
+					.ForegroundColor_Lambda([this]() { return IsPlaying.Get() ? FStyleColors::AccentRed : FSlateColor::UseStyle(); })
+					.OnClicked_Lambda([this]() { return StopClickedDelegate.Execute(*this); })
+					.ToolTipText(StopItemTooltipText)
+					[
+						SNew(SImage)
+						.Image(FSequencerPlaylistsStyle::Get().GetBrush("SequencerPlaylists.Stop.Small"))
+						.ColorAndOpacity(FSlateColor::UseForeground())
+					]
 				]
-			]
-			+ SVerticalBox::Slot()
-			.Padding(8.0f, 2.0f, 0.0f, 4.0f)
-			[
-				SNew(SButton)
-				.ButtonStyle(FSequencerPlaylistsStyle::Get(), "SequencerPlaylists.HoverTransport.Reset")
-				.ContentPadding(0)
-				.HAlign(HAlign_Center)
-				.VAlign(VAlign_Center)
-				.Visibility_Lambda([this]() { return (IsHovered() && !InPlayMode()) ? EVisibility::Visible : EVisibility::Hidden; })
-				.OnClicked_Lambda([this]() { return ResetClickedDelegate.Execute(*this); })
-				.ToolTipText(ResetItemTooltipText)
+				+ SHorizontalBox::Slot()
+				.Padding(2.0f, 0.0f, 0.0f, 0.0f)
 				[
-					SNew(SImage)
-					.Image(FSequencerPlaylistsStyle::Get().GetBrush("SequencerPlaylists.Reset.Small"))
-					.ColorAndOpacity(FSlateColor::UseForeground())
+					SNew(SButton)
+					.ButtonStyle(FSequencerPlaylistsStyle::Get(), "SequencerPlaylists.HoverTransport.Reset")
+					.ContentPadding(FMargin(0.0f, 0.0f, 0.0f, 1.0f))
+					.HAlign(HAlign_Center)
+					.VAlign(VAlign_Center)
+					.Visibility_Lambda([this]() { return (IsHovered() && !InPlayMode()) ? EVisibility::Visible : EVisibility::Hidden; })
+					.OnClicked_Lambda([this]() { return ResetClickedDelegate.Execute(*this); })
+					.ToolTipText(ResetItemTooltipText)
+					[
+						SNew(SImage)
+						.Image(FSequencerPlaylistsStyle::Get().GetBrush("SequencerPlaylists.Reset.Small"))
+						.ColorAndOpacity(FSlateColor::UseForeground())
+					]
 				]
 			];
 	}
@@ -1420,7 +1430,7 @@ TSharedRef<SWidget> SSequencerPlaylistItemWidget::GenerateWidgetForColumn(const
 			[
 				SNew(SImage)
 				.Image(FSequencerPlaylistsStyle::Get().GetBrush("SequencerPlaylists.Pause"))
-				//.ColorAndOpacity(FStyleColors::AccentBlue)
+				.ColorAndOpacity(FSlateColor::UseForeground())
 			];
 
 		return SNew(SVerticalBox)
