Sequencer Playlists: Don't clear pending take / sources unnecessarily.

diff --git a/Engine/Plugins/VirtualProduction/SequencerPlaylists/Source/SequencerPlaylists/Private/SequencerPlaylistPlayer.cpp b/Engine/Plugins/VirtualProduction/SequencerPlaylists/Source/SequencerPlaylists/Private/SequencerPlaylistPlayer.cpp
--- a/Engine/Plugins/VirtualProduction/SequencerPlaylists/Source/SequencerPlaylists/Private/SequencerPlaylistPlayer.cpp
+++ b/Engine/Plugins/VirtualProduction/SequencerPlaylists/Source/SequencerPlaylists/Private/SequencerPlaylistPlayer.cpp
@@ -470,10 +470,13 @@ TSharedPtr<ISequencer> USequencerPlaylistPlayer::GetOrCreateSequencer()
 	{
 		UTakePreset* Preset = UTakePreset::AllocateTransientPreset(GetDefault<UTakeRecorderUserSettings>()->LastOpenedPreset.Get());
 
-		FScopedTransaction Transaction(LOCTEXT("CreateEmptyTake", "Create Empty Playlist Sequence"));
+		if (!Preset->GetLevelSequence())
+		{
+			FScopedTransaction Transaction(LOCTEXT("CreateEmptyTake", "Create Empty Playlist Sequence"));
 
-		Preset->Modify();
-		Preset->CreateLevelSequence();
+			Preset->Modify();
+			Preset->CreateLevelSequence();
+		}
 
 		RootSequence = Preset->GetLevelSequence();
 	}
