Sequencer Playlists: Partially restore prior approach to initiating holds.

Reliably mapping between state in the pending take to the equivalent in the destination take turns out to be difficult due to some missing lifecycle events in Take Recorder. This is a conservative fix that ensures the UI behaves correctly after Take Recorder has started. The tradeoff is that, as in previous releases, the effect on the scene of any sequence held prior to beginning recording will not be visible until recording begins.

diff --git a/Engine/Plugins/VirtualProduction/SequencerPlaylists/Source/SequencerPlaylists/Private/SequencerPlaylistItem_Sequence.cpp b/Engine/Plugins/VirtualProduction/SequencerPlaylists/Source/SequencerPlaylists/Private/SequencerPlaylistItem_Sequence.cpp
--- a/Engine/Plugins/VirtualProduction/SequencerPlaylists/Source/SequencerPlaylists/Private/SequencerPlaylistItem_Sequence.cpp
+++ b/Engine/Plugins/VirtualProduction/SequencerPlaylists/Source/SequencerPlaylists/Private/SequencerPlaylistItem_Sequence.cpp
@@ -140,15 +140,22 @@ FSequencerPlaylistItemPlayer_Sequence::GetPlaybackState(USequencerPlaylistItem*
 		return Result;
 	}
 
-	const FItemState* ItemState = ItemStates.Find(Item);
-	if (!ItemState)
+	if (IsSequencerRecordingOrPlaying())
 	{
-		return Result;
-	}
+		const FItemState* ItemState = ItemStates.Find(Item);
+		if (!ItemState)
+		{
+			return Result;
+		}
 
-	Result.bIsPlaying = ItemState->PlayingUntil_RootTicks > Sequencer->GetGlobalTime().Time.FloorToFrame().Value;
-	Result.bIsPaused = ItemState->WeakHoldSection.IsValid();
-	Result.PlaybackDirection = ItemState->LastPlayDirection;
+		Result.bIsPlaying = ItemState->PlayingUntil_RootTicks > Sequencer->GetGlobalTime().Time.FloorToFrame().Value;
+		Result.bIsPaused = ItemState->WeakHoldSection.IsValid();
+		Result.PlaybackDirection = ItemState->LastPlayDirection;
+	}
+	else
+	{
+		Result.bIsPaused = Item->bHoldAtFirstFrame;
+	}
 
 	return Result;
 }
@@ -271,13 +278,10 @@ bool FSequencerPlaylistItemPlayer_Sequence::InternalPause(USequencerPlaylistItem
 		return bSequenceWasModified;
 	}
 
-	const FQualifiedFrameTime GlobalTime = Sequencer->GetGlobalTime();
-
-	FItemState& ItemState = ItemStates.FindOrAdd(Item);
-	UMovieSceneSubSection* HoldSection = ItemState.WeakHoldSection.Get();
-	if (HoldSection)
+	if (IsSequencerRecordingOrPlaying())
 	{
-		if (IsSequencerRecordingOrPlaying())
+		FItemState& ItemState = ItemStates.FindOrAdd(Item);
+		if (UMovieSceneSubSection* HoldSection = ItemState.WeakHoldSection.Get())
 		{
 			// Resume playback
 			FPlayParams PlayParams;
@@ -287,16 +291,9 @@ bool FSequencerPlaylistItemPlayer_Sequence::InternalPause(USequencerPlaylistItem
 		}
 		else
 		{
-			// Cancel hold
-			bSequenceWasModified |= Stop(Item);
-			ItemState.bHoldOnReset = false;
-		}
-	}
-	else
-	{
-		FHoldParams HoldParams;
-		if (IsSequencerRecordingOrPlaying())
-		{
+			FHoldParams HoldParams;
+
+			const FQualifiedFrameTime GlobalTime = Sequencer->GetGlobalTime();
 			UMovieSceneSubSection* CurrentPlayingSection = IsTimeWithinAnySection(GlobalTime, ItemState.WeakPlaySections);
 			if (CurrentPlayingSection)
 			{
@@ -305,14 +302,14 @@ bool FSequencerPlaylistItemPlayer_Sequence::InternalPause(USequencerPlaylistItem
 
 				HoldParams.StartFrameOffset_SceneTicks = InnerTimeAndWarp.Get<0>().GetValue();
 			}
-		}
-		else
-		{
-			ItemState.bHoldOnReset = true;
-		}
 
-		bSequenceWasModified |= Stop(Item);
-		bSequenceWasModified |= InternalAddHold(Item, HoldParams);
+			bSequenceWasModified |= Stop(Item);
+			bSequenceWasModified |= InternalAddHold(Item, HoldParams);
+		}
+	}
+	else
+	{
+		Item->bHoldAtFirstFrame = !Item->bHoldAtFirstFrame;
 	}
 
 	return bSequenceWasModified;
@@ -422,15 +419,16 @@ bool FSequencerPlaylistItemPlayer_Sequence::InternalReset(USequencerPlaylistItem
 {
 	bool bSequenceWasModified = false;
 
-	FItemState& ItemState = ItemStates.FindOrAdd(Item);
-	if (ItemState.bHoldOnReset)
+	if (!Item)
 	{
-		bSequenceWasModified |= Stop(Item);
-		bSequenceWasModified |= AddHold(Item);
+		return bSequenceWasModified;
 	}
-	else
+
+	bSequenceWasModified |= Stop(Item);
+
+	if (Item->bHoldAtFirstFrame)
 	{
-		bSequenceWasModified |= Stop(Item);
+		bSequenceWasModified |= AddHold(Item);
 	}
 
 	return bSequenceWasModified;
diff --git a/Engine/Plugins/VirtualProduction/SequencerPlaylists/Source/SequencerPlaylists/Private/SequencerPlaylistItem_Sequence.h b/Engine/Plugins/VirtualProduction/SequencerPlaylists/Source/SequencerPlaylists/Private/SequencerPlaylistItem_Sequence.h
--- a/Engine/Plugins/VirtualProduction/SequencerPlaylists/Source/SequencerPlaylists/Private/SequencerPlaylistItem_Sequence.h
+++ b/Engine/Plugins/VirtualProduction/SequencerPlaylists/Source/SequencerPlaylists/Private/SequencerPlaylistItem_Sequence.h
@@ -37,7 +37,6 @@ class FSequencerPlaylistItemPlayer_Sequence : public ISequencerPlaylistItemPlaye
 		TArray<TWeakObjectPtr<UMovieSceneSubSection>> WeakPlaySections;
 		int32 PlayingUntil_RootTicks = TNumericLimits<int32>::Min();
 		ESequencerPlaylistPlaybackDirection LastPlayDirection = ESequencerPlaylistPlaybackDirection::Forward;
-		bool bHoldOnReset = false;
 	};
 
 public:
diff --git a/Engine/Plugins/VirtualProduction/SequencerPlaylists/Source/SequencerPlaylists/Private/SequencerPlaylistPlayer.cpp b/Engine/Plugins/VirtualProduction/SequencerPlaylists/Source/SequencerPlaylists/Private/SequencerPlaylistPlayer.cpp
--- a/Engine/Plugins/VirtualProduction/SequencerPlaylists/Source/SequencerPlaylists/Private/SequencerPlaylistPlayer.cpp
+++ b/Engine/Plugins/VirtualProduction/SequencerPlaylists/Source/SequencerPlaylists/Private/SequencerPlaylistPlayer.cpp
@@ -165,18 +165,24 @@ bool USequencerPlaylistPlayer::IsPlaying(USequencerPlaylistItem* Item)
 
 FSequencerPlaylistPlaybackState USequencerPlaylistPlayer::GetPlaybackState(USequencerPlaylistItem* Item)
 {
+	FSequencerPlaylistPlaybackState Result;
+
 	if (!Item)
 	{
-		return FSequencerPlaylistPlaybackState();
+		return Result;
 	}
 
 	// If Sequencer isn't already open, don't open it
-	if (!WeakSequencer.IsValid())
+	if (WeakSequencer.IsValid())
 	{
-		return FSequencerPlaylistPlaybackState();
+		Result = GetCheckedItemPlayer(Item)->GetPlaybackState(Item);
+	}
+	else
+	{
+		Result.bIsPaused = Item->bHoldAtFirstFrame;
 	}
 
-	return GetCheckedItemPlayer(Item)->GetPlaybackState(Item);
+	return Result;
 }
 
 
@@ -521,6 +527,23 @@ void USequencerPlaylistPlayer::OnTakeRecorderInitialized(UTakeRecorder* InRecord
 
 void USequencerPlaylistPlayer::OnTakeRecorderStarted(UTakeRecorder* InRecorder)
 {
+	if (!ensure(Playlist))
+	{
+		return;
+	}
+
+	if (TSharedPtr<ISequencer> Sequencer = GetOrCreateSequencer())
+	{
+		FScopedTransaction Transaction(LOCTEXT("TakeRecorderStartedTransaction", "Playlist - Take Recorder started"));
+
+		for (USequencerPlaylistItem* Item : Playlist->Items)
+		{
+			if (Item->bHoldAtFirstFrame && !Item->bMute)
+			{
+				GetCheckedItemPlayer(Item)->AddHold(Item);
+			}
+		}
+	}
 }
 
 
diff --git a/Engine/Plugins/VirtualProduction/SequencerPlaylists/Source/SequencerPlaylists/Private/SequencerPlaylistsStyle.cpp b/Engine/Plugins/VirtualProduction/SequencerPlaylists/Source/SequencerPlaylists/Private/SequencerPlaylistsStyle.cpp
--- a/Engine/Plugins/VirtualProduction/SequencerPlaylists/Source/SequencerPlaylists/Private/SequencerPlaylistsStyle.cpp
+++ b/Engine/Plugins/VirtualProduction/SequencerPlaylists/Source/SequencerPlaylists/Private/SequencerPlaylistsStyle.cpp
@@ -81,11 +81,11 @@ TSharedRef<FSlateStyleSet> FSequencerPlaylistsStyle::Create()
 	DimColor.A = 0.5f;
 	Style->Set("SequencerPlaylists.Item.Dim", new FSlateColorBrush(DimColor));
 
-	Style->Set("SequencerPlaylists.NewPlaylist", new IMAGE_BRUSH_SVG("PlaylistNew", Icon16x16));
+	Style->Set("SequencerPlaylists.NewPlaylist", new IMAGE_BRUSH_SVG("PlaylistNew", Icon20x20));
 	// SavePlaylist is registered under a different root below.
-	Style->Set("SequencerPlaylists.OpenPlaylist", new IMAGE_BRUSH_SVG("PlaylistLoad", Icon16x16));
+	Style->Set("SequencerPlaylists.OpenPlaylist", new IMAGE_BRUSH_SVG("PlaylistLoad", Icon20x20));
 
-	Style->Set("SequencerPlaylists.PlayMode", new IMAGE_BRUSH_SVG("PlaybackAuto", Icon16x16));
+	Style->Set("SequencerPlaylists.PlayMode", new IMAGE_BRUSH_SVG("PlaybackAuto", Icon20x20));
 
 	Style->Set("SequencerPlaylists.HoldFrame", new IMAGE_BRUSH_SVG("HoldFrame", Icon16x16));
 
@@ -102,7 +102,7 @@ TSharedRef<FSlateStyleSet> FSequencerPlaylistsStyle::Create()
 	{
 		const FToolBarStyle& SlimToolbarStyle = FAppStyle::Get().GetWidgetStyle<FToolBarStyle>("SlimToolBar");
 		FToolBarStyle MainToolbarStyle = FToolBarStyle(SlimToolbarStyle)
-			.SetIconSize(Icon16x16)
+			.SetIconSize(Icon20x20)
 			.SetComboButtonPadding(FMargin(2.0f, 0.0f));
 
 		FButtonStyle MainToolbarButtonStyle = FButtonStyle(SlimToolbarStyle.ButtonStyle)
@@ -281,7 +281,7 @@ TSharedRef<FSlateStyleSet> FSequencerPlaylistsStyle::Create()
 			Style->SetContentRoot(EngineEditorSlateDir);
 			Style->SetCoreContentRoot(FPaths::EngineContentDir() / TEXT("Slate"));
 
-			Style->Set("SequencerPlaylists.SavePlaylist", new IMAGE_BRUSH_SVG("Starship/Common/SaveCurrent", Icon16x16));
+			Style->Set("SequencerPlaylists.SavePlaylist", new IMAGE_BRUSH_SVG("Starship/Common/SaveCurrent", Icon20x20));
 
 			Style->Set("SequencerPlaylists.Play", new IMAGE_BRUSH_SVG("Starship/Common/play", Icon20x20));
 			Style->Set("SequencerPlaylists.Play.Small", new IMAGE_BRUSH_SVG("Starship/Common/play", Icon16x16));
diff --git a/Engine/Plugins/VirtualProduction/SequencerPlaylists/Source/SequencerPlaylists/Private/SequencerPlaylistsWidgets.cpp b/Engine/Plugins/VirtualProduction/SequencerPlaylists/Source/SequencerPlaylists/Private/SequencerPlaylistsWidgets.cpp
--- a/Engine/Plugins/VirtualProduction/SequencerPlaylists/Source/SequencerPlaylists/Private/SequencerPlaylistsWidgets.cpp
+++ b/Engine/Plugins/VirtualProduction/SequencerPlaylists/Source/SequencerPlaylists/Private/SequencerPlaylistsWidgets.cpp
@@ -381,6 +381,7 @@ TSharedRef<SWidget> SSequencerPlaylistPanel::Construct_Transport()
 			[
 				SNew(SImage)
 				.Image(FSequencerPlaylistsStyle::Get().GetBrush("SequencerPlaylists.Stop"))
+				.DesiredSizeOverride(FVector2D(18.0f, 18.0f))
 				.ColorAndOpacity(FSlateColor::UseForeground())
 			]
 		]
@@ -1227,6 +1228,7 @@ void SSequencerPlaylistItemWidget::ConstructChildren(ETableViewMode::Type InOwne
 				[
 					SNew(SImage)
 					.Image(FSequencerPlaylistsStyle::Get().GetBrush("SequencerPlaylists.Stop"))
+					.DesiredSizeOverride(FVector2D(18.0f, 18.0f))
 					.ColorAndOpacity(FSlateColor::UseForeground())
 				]
 			]
@@ -1330,6 +1332,7 @@ TSharedRef<SWidget> SSequencerPlaylistItemWidget::GenerateWidgetForColumn(const
 					[
 						SNew(SImage)
 						.Image(FSequencerPlaylistsStyle::Get().GetBrush("SequencerPlaylists.Stop.Small"))
+						.DesiredSizeOverride(FVector2D(14.0f, 14.0f))
 						.ColorAndOpacity(FSlateColor::UseForeground())
 					]
 				]
diff --git a/Engine/Plugins/VirtualProduction/SequencerPlaylists/Source/SequencerPlaylists/Public/SequencerPlaylistItem.h b/Engine/Plugins/VirtualProduction/SequencerPlaylists/Source/SequencerPlaylists/Public/SequencerPlaylistItem.h
--- a/Engine/Plugins/VirtualProduction/SequencerPlaylists/Source/SequencerPlaylists/Public/SequencerPlaylistItem.h
+++ b/Engine/Plugins/VirtualProduction/SequencerPlaylists/Source/SequencerPlaylists/Public/SequencerPlaylistItem.h
@@ -23,15 +23,13 @@ public:
 	UPROPERTY(EditAnywhere, BlueprintReadWrite, Category="Sequencer Playlists")
 	int32 EndFrameOffset;
 
-#if WITH_EDITORONLY_DATA
 	/**
 	 * If true, the sequence will be inserted immediately on recording start and any time Reset()
 	 * is called, paused at the first frame indefinitely until subsequently played or stopped.
+	 * This is controlled by toggling pause prior to recording/playback.
 	 */
-	UE_DEPRECATED(5.2, "Holds are now created by invoking Pause prior to beginning recording.")
-	UPROPERTY(meta=(DeprecatedProperty, DeprecationMessage="Holds are now created by invoking Pause prior to beginning recording."))
-	bool bHoldAtFirstFrame_DEPRECATED;
-#endif
+	UPROPERTY()
+	bool bHoldAtFirstFrame;
 
 	/** 0 is single playthrough, >= 1 is (n+1) playthroughs. */
 	UPROPERTY(EditAnywhere, BlueprintReadWrite, Category="Sequencer Playlists")
