Sequencer Playlists: Integrate multi-user level sequence precaching.

diff --git a/Engine/Plugins/VirtualProduction/SequencerPlaylists/SequencerPlaylists.uplugin b/Engine/Plugins/VirtualProduction/SequencerPlaylists/SequencerPlaylists.uplugin
--- a/Engine/Plugins/VirtualProduction/SequencerPlaylists/SequencerPlaylists.uplugin
+++ b/Engine/Plugins/VirtualProduction/SequencerPlaylists/SequencerPlaylists.uplugin
@@ -15,6 +15,14 @@
 	"IsExperimentalVersion": false,
 	"Installed": false,
 	"Plugins": [
+		{
+			"Name": "ConcertSyncClient",
+			"Enabled": true
+		},
+		{
+			"Name": "ConcertSyncCore",
+			"Enabled": true
+		},
 		{
 			"Name": "Takes",
 			"Enabled": true
diff --git a/Engine/Plugins/VirtualProduction/SequencerPlaylists/Source/SequencerPlaylists/Private/SequencerPlaylistItem_Sequence.cpp b/Engine/Plugins/VirtualProduction/SequencerPlaylists/Source/SequencerPlaylists/Private/SequencerPlaylistItem_Sequence.cpp
--- a/Engine/Plugins/VirtualProduction/SequencerPlaylists/Source/SequencerPlaylists/Private/SequencerPlaylistItem_Sequence.cpp
+++ b/Engine/Plugins/VirtualProduction/SequencerPlaylists/Source/SequencerPlaylists/Private/SequencerPlaylistItem_Sequence.cpp
@@ -8,6 +8,7 @@
 #include "MovieSceneFolder.h"
 #include "Recorder/TakeRecorder.h"
 #include "SequencerPlaylist.h"
+#include "SequencerPlaylistsSubsystem.h"
 #include "TrackEditors/SubTrackEditorBase.h" // for FSubTrackEditorUtil::CanAddSubSequence
 #include "Tracks/MovieSceneSubTrack.h"
 #include "Widgets/Notifications/SNotificationList.h"
@@ -66,12 +67,52 @@ namespace UE::Private::SequencerPlaylistItem_Sequence
 }
 
 
+/*static*/ FName USequencerPlaylistItem_Sequence::GetSequencePropertyName()
+{
+	static const FName SequencePropertyName =
+		GET_MEMBER_NAME_CHECKED(USequencerPlaylistItem_Sequence, Sequence);
+
+	return SequencePropertyName;
+}
+
+
 FText USequencerPlaylistItem_Sequence::GetDisplayName()
 {
 	return Sequence ? Sequence->GetDisplayName() : LOCTEXT("SequenceItemNullDisplayName", "(No sequence)");
 }
 
 
+void USequencerPlaylistItem_Sequence::PostEditChangeProperty(FPropertyChangedEvent& PropertyChangedEvent)
+{
+	Super::PostEditChangeProperty(PropertyChangedEvent);
+
+	const FName ChangedPropertyName = PropertyChangedEvent.MemberProperty
+		? PropertyChangedEvent.MemberProperty->GetFName()
+		: NAME_None;
+
+	if (ChangedPropertyName == NAME_None || ChangedPropertyName == GetSequencePropertyName())
+	{
+		USequencerPlaylistsSubsystem* Subsystem =
+			GEditor->GetEditorSubsystem<USequencerPlaylistsSubsystem>();
+		if (Subsystem)
+		{
+			Subsystem->UpdatePrecacheSet();
+		}
+	}
+}
+
+
+void USequencerPlaylistItem_Sequence::SetSequence(ULevelSequence* NewSequence)
+{
+	Sequence = NewSequence;
+
+	if (USequencerPlaylistsSubsystem* Subsystem = GEditor->GetEditorSubsystem<USequencerPlaylistsSubsystem>())
+	{
+		Subsystem->UpdatePrecacheSet();
+	}
+}
+
+
 FSequencerPlaylistItemPlayer_Sequence::FSequencerPlaylistItemPlayer_Sequence(TSharedRef<ISequencer> Sequencer)
 	: WeakSequencer(Sequencer)
 {
@@ -129,7 +170,7 @@ FSequencerPlaylistItemPlayer_Sequence::GetPlaybackState(USequencerPlaylistItem*
 	FSequencerPlaylistPlaybackState Result;
 
 	USequencerPlaylistItem_Sequence* SequenceItem = CastChecked<USequencerPlaylistItem_Sequence>(Item);
-	if (!SequenceItem || !SequenceItem->Sequence)
+	if (!SequenceItem || !SequenceItem->GetSequence())
 	{
 		return Result;
 	}
@@ -167,7 +208,7 @@ bool FSequencerPlaylistItemPlayer_Sequence::InternalPlay(
 )
 {
 	USequencerPlaylistItem_Sequence* SequenceItem = CastChecked<USequencerPlaylistItem_Sequence>(Item);
-	if (!SequenceItem || !SequenceItem->Sequence)
+	if (!SequenceItem || !SequenceItem->GetSequence())
 	{
 		return false;
 	}
@@ -181,9 +222,9 @@ bool FSequencerPlaylistItemPlayer_Sequence::InternalPlay(
 	ULevelSequence* RootSequence = Cast<ULevelSequence>(Sequencer->GetRootMovieSceneSequence());
 	UMovieScene* RootScene = RootSequence->GetMovieScene();
 
-	if (!FSubTrackEditorUtil::CanAddSubSequence(RootSequence, *SequenceItem->Sequence))
+	if (!FSubTrackEditorUtil::CanAddSubSequence(RootSequence, *SequenceItem->GetSequence()))
 	{
-		FNotificationInfo Info(FText::Format(LOCTEXT("InvalidSequence", "Invalid level sequence {0}. There could be a circular dependency."), SequenceItem->Sequence->GetDisplayName()));
+		FNotificationInfo Info(FText::Format(LOCTEXT("InvalidSequence", "Invalid level sequence {0}. There could be a circular dependency."), SequenceItem->GetSequence()->GetDisplayName()));
 		Info.bUseLargeFont = false;
 		FSlateNotificationManager::Get().AddNotification(Info);
 		return false;
@@ -203,7 +244,7 @@ bool FSequencerPlaylistItemPlayer_Sequence::InternalPlay(
 
 	ItemState.LastPlayDirection = PlayParams.Direction;
 
-	UMovieScene* PlayScene = SequenceItem->Sequence->GetMovieScene();
+	UMovieScene* PlayScene = SequenceItem->GetSequence()->GetMovieScene();
 	const TRange<FFrameNumber> SequencePlayRange = PlayScene->GetPlaybackRange();
 
 	const FFrameTime StartFrameOffset_Ticks = PlayParams.StartFrameOffset_SceneTicks.IsSet()
@@ -236,7 +277,7 @@ bool FSequencerPlaylistItemPlayer_Sequence::InternalPlay(
 
 	ItemState.PlayingUntil_RootTicks = FMath::Max(StartFrame.Value + Duration, ItemState.PlayingUntil_RootTicks);
 
-	UMovieSceneSubSection* WorkingSubSection = WorkingTrack->AddSequence(SequenceItem->Sequence, StartFrame, Duration);
+	UMovieSceneSubSection* WorkingSubSection = WorkingTrack->AddSequence(SequenceItem->GetSequence(), StartFrame, Duration);
 
 	WorkingSubSection->Parameters.TimeScale = TimeScale;
 	WorkingSubSection->Parameters.StartFrameOffset = StartFrameOffset_Ticks.FloorToFrame();
@@ -267,7 +308,7 @@ bool FSequencerPlaylistItemPlayer_Sequence::InternalPause(USequencerPlaylistItem
 	bool bSequenceWasModified = false;
 
 	USequencerPlaylistItem_Sequence* SequenceItem = CastChecked<USequencerPlaylistItem_Sequence>(Item);
-	if (!SequenceItem || !SequenceItem->Sequence)
+	if (!SequenceItem || !SequenceItem->GetSequence())
 	{
 		return bSequenceWasModified;
 	}
@@ -321,7 +362,7 @@ bool FSequencerPlaylistItemPlayer_Sequence::InternalStop(USequencerPlaylistItem*
 	bool bSequenceWasModified = false;
 
 	USequencerPlaylistItem_Sequence* SequenceItem = CastChecked<USequencerPlaylistItem_Sequence>(Item);
-	if (!SequenceItem || !SequenceItem->Sequence)
+	if (!SequenceItem || !SequenceItem->GetSequence())
 	{
 		return bSequenceWasModified;
 	}
@@ -358,7 +399,7 @@ bool FSequencerPlaylistItemPlayer_Sequence::InternalStop(USequencerPlaylistItem*
 bool FSequencerPlaylistItemPlayer_Sequence::InternalAddHold(USequencerPlaylistItem* Item, const FHoldParams& HoldParams)
 {
 	USequencerPlaylistItem_Sequence* SequenceItem = CastChecked<USequencerPlaylistItem_Sequence>(Item);
-	if (!SequenceItem || !SequenceItem->Sequence)
+	if (!SequenceItem || !SequenceItem->GetSequence())
 	{
 		return false;
 	}
@@ -379,9 +420,9 @@ bool FSequencerPlaylistItemPlayer_Sequence::InternalAddHold(USequencerPlaylistIt
 	ULevelSequence* RootSequence = Cast<ULevelSequence>(Sequencer->GetRootMovieSceneSequence());
 	UMovieScene* RootScene = RootSequence->GetMovieScene();
 
-	if (!FSubTrackEditorUtil::CanAddSubSequence(RootSequence, *SequenceItem->Sequence))
+	if (!FSubTrackEditorUtil::CanAddSubSequence(RootSequence, *SequenceItem->GetSequence()))
 	{
-		FNotificationInfo Info(FText::Format(LOCTEXT("InvalidSequence", "Invalid level sequence {0}. There could be a circular dependency."), SequenceItem->Sequence->GetDisplayName()));
+		FNotificationInfo Info(FText::Format(LOCTEXT("InvalidSequence", "Invalid level sequence {0}. There could be a circular dependency."), SequenceItem->GetSequence()->GetDisplayName()));
 		Info.bUseLargeFont = false;
 		FSlateNotificationManager::Get().AddNotification(Info);
 		return false;
@@ -391,7 +432,7 @@ bool FSequencerPlaylistItemPlayer_Sequence::InternalAddHold(USequencerPlaylistIt
 
 	UMovieSceneSubTrack* WorkingTrack = GetOrCreateWorkingTrack(Item);
 
-	ULevelSequence* HoldSequence = SequenceItem->Sequence;
+	ULevelSequence* HoldSequence = SequenceItem->GetSequence();
 	UMovieScene* HoldScene = HoldSequence->GetMovieScene();
 
 	const FQualifiedFrameTime GlobalTime = Sequencer->GetGlobalTime();
@@ -472,7 +513,7 @@ void FSequencerPlaylistItemPlayer_Sequence::ClearItemStates()
 UMovieSceneSubTrack* FSequencerPlaylistItemPlayer_Sequence::GetOrCreateWorkingTrack(USequencerPlaylistItem* Item)
 {
 	USequencerPlaylistItem_Sequence* SequenceItem = CastChecked<USequencerPlaylistItem_Sequence>(Item);
-	if (!SequenceItem || !SequenceItem->Sequence)
+	if (!SequenceItem || !SequenceItem->GetSequence())
 	{
 		return nullptr;
 	}
@@ -495,7 +536,7 @@ UMovieSceneSubTrack* FSequencerPlaylistItemPlayer_Sequence::GetOrCreateWorkingTr
 	}
 
 	UMovieSceneSubTrack* NewWorkingTrack = RootScene->AddTrack<UMovieSceneSubTrack>();
-	NewWorkingTrack->SetDisplayName(FText::Format(LOCTEXT("SequenceItemTrackName", "Item - {0}"), SequenceItem->Sequence->GetDisplayName()));
+	NewWorkingTrack->SetDisplayName(FText::Format(LOCTEXT("SequenceItemTrackName", "Item - {0}"), SequenceItem->GetSequence()->GetDisplayName()));
 
 	// Find or create folder named for our playlist, and organize our track beneath it.
 	FText PlaylistName = FText::GetEmpty();
diff --git a/Engine/Plugins/VirtualProduction/SequencerPlaylists/Source/SequencerPlaylists/Private/SequencerPlaylistItem_Sequence.h b/Engine/Plugins/VirtualProduction/SequencerPlaylists/Source/SequencerPlaylists/Private/SequencerPlaylistItem_Sequence.h
--- a/Engine/Plugins/VirtualProduction/SequencerPlaylists/Source/SequencerPlaylists/Private/SequencerPlaylistItem_Sequence.h
+++ b/Engine/Plugins/VirtualProduction/SequencerPlaylists/Source/SequencerPlaylists/Private/SequencerPlaylistItem_Sequence.h
@@ -20,10 +20,24 @@ class USequencerPlaylistItem_Sequence : public USequencerPlaylistItem
 {
 	GENERATED_BODY()
 
-	FText GetDisplayName() override;
-
 public:
-	UPROPERTY(EditAnywhere, BlueprintReadWrite, Category="Sequencer Playlists", meta=(NoResetToDefault))
+	static FName GetSequencePropertyName();
+
+	//~ Begin USequencerPlaylistItem interface
+	virtual FText GetDisplayName() override;
+	//~ End USequencerPlaylistItem interface
+
+	//~ Begin UObject interface
+	virtual void PostEditChangeProperty(struct FPropertyChangedEvent& PropertyChangedEvent);
+	//~ End UObject interface
+
+	UFUNCTION(BlueprintCallable, Category="Sequencer Playlists")
+	void SetSequence(ULevelSequence* NewSequence);
+
+	ULevelSequence* GetSequence() const { return Sequence; }
+
+protected:
+	UPROPERTY(EditAnywhere, BlueprintReadWrite, Setter=SetSequence, Category="Sequencer Playlists", meta=(NoResetToDefault))
 	TObjectPtr<ULevelSequence> Sequence;
 };
 
@@ -41,7 +55,7 @@ class FSequencerPlaylistItemPlayer_Sequence : public ISequencerPlaylistItemPlaye
 
 public:
 	FSequencerPlaylistItemPlayer_Sequence(TSharedRef<ISequencer> Sequencer);
-	~FSequencerPlaylistItemPlayer_Sequence() override;
+	virtual ~FSequencerPlaylistItemPlayer_Sequence() override;
 
 	//~ Begin ISequencerPlaylistItemPlayer
 	virtual bool Play(USequencerPlaylistItem* Item,
diff --git a/Engine/Plugins/VirtualProduction/SequencerPlaylists/Source/SequencerPlaylists/Private/SequencerPlaylistsSubsystem.cpp b/Engine/Plugins/VirtualProduction/SequencerPlaylists/Source/SequencerPlaylists/Private/SequencerPlaylistsSubsystem.cpp
--- a/Engine/Plugins/VirtualProduction/SequencerPlaylists/Source/SequencerPlaylists/Private/SequencerPlaylistsSubsystem.cpp
+++ b/Engine/Plugins/VirtualProduction/SequencerPlaylists/Source/SequencerPlaylists/Private/SequencerPlaylistsSubsystem.cpp
@@ -1,14 +1,34 @@
 // Copyright Epic Games, Inc. All Rights Reserved.
 
 #include "SequencerPlaylistsSubsystem.h"
+#include "ConcertSequencerMessages.h"
+#include "IConcertClient.h"
+#include "IConcertSyncClient.h"
+#include "IConcertSyncClientModule.h"
+#include "IConcertSession.h"
+#include "LevelSequence.h"
 #include "SequencerPlaylist.h"
+#include "SequencerPlaylistItem_Sequence.h"
 #include "SequencerPlaylistPlayer.h"
+#include "SequencerPlaylistsLog.h"
 #include "UObject/Package.h"
 
 
 void USequencerPlaylistsSubsystem::Initialize(FSubsystemCollectionBase& Collection)
 {
 	Super::Initialize(Collection);
+
+	if (TSharedPtr<IConcertSyncClient> ConcertSyncClient = IConcertSyncClientModule::Get().GetClient(TEXT("MultiUser")))
+	{
+		IConcertClientRef ConcertClient = ConcertSyncClient->GetConcertClient();
+		ConcertClient->OnSessionStartup().AddUObject(this, &USequencerPlaylistsSubsystem::OnConcertSessionStartup);
+		ConcertClient->OnSessionShutdown().AddUObject(this, &USequencerPlaylistsSubsystem::OnConcertSessionShutdown);
+
+		if (TSharedPtr<IConcertClientSession> ConcertClientSession = ConcertClient->GetCurrentSession())
+		{
+			OnConcertSessionStartup(ConcertClientSession.ToSharedRef());
+		}
+	}
 }
 
 
@@ -34,7 +54,6 @@ USequencerPlaylistPlayer* USequencerPlaylistsSubsystem::CreatePlayerForEditor(TS
 
 	ensure(EditorPackages.Find(EditorHandle) == nullptr);
 	ensure(EditorPlaylists.Find(EditorHandle) == nullptr);
-	ensure(EditorPlayers.Find(EditorHandle) == nullptr);
 
 	USequencerPlaylist* NewPlaylist = CreateTransientPlaylistForEditor(Editor);
 	EditorPlaylists.Add(EditorHandle, NewPlaylist);
@@ -52,12 +71,7 @@ USequencerPlaylist* USequencerPlaylistsSubsystem::CreateTransientPlaylistForEdit
 {
 	FName PackageName = ::MakeUniqueObjectName(nullptr, UPackage::StaticClass(), TEXT("/Engine/Transient/SequencerPlaylist"));
 	UPackage* PlaylistPackage = NewObject<UPackage>(nullptr, PackageName, RF_Transient | RF_Transactional);
-	//PlaylistPackage->AddToRoot();
-
-	USequencerPlaylist* Playlist = NewObject<USequencerPlaylist>(PlaylistPackage, TEXT("UntitledPlaylist"));
-	Playlist->SetFlags(RF_Transactional);
-
-	return Playlist;
+	return NewObject<USequencerPlaylist>(PlaylistPackage, TEXT("UntitledPlaylist"), RF_Transactional);
 }
 
 
@@ -70,4 +84,86 @@ void USequencerPlaylistsSubsystem::NotifyEditorClosed(SSequencerPlaylistPanel* E
 	ensure(EditorPlaylists.Remove(EditorHandle));
 	ensure(EditorPlayers.Remove(EditorHandle));
 	ensure(EditorPackages.Remove(EditorHandle));
+
+	UpdatePrecacheSet();
+}
+
+
+void USequencerPlaylistsSubsystem::UpdatePrecacheSet()
+{
+	// Determine which sequences (if any) were added to/removed from the precache set.
+	// We copy the previous set, recompute the current set, then compare them.
+	const TSet<TObjectPtr<ULevelSequence>> PreviousPrecache = PrecacheSequences;
+
+	PrecacheSequences.Empty(PrecacheSequences.Num());
+
+	using FPlaylistPair = TPair<FSequencerPlaylistEditorHandle, TObjectPtr<USequencerPlaylist>>;
+	for (const FPlaylistPair& EditorPlaylist : EditorPlaylists)
+	{
+		for (USequencerPlaylistItem* Item : EditorPlaylist.Value->Items)
+		{
+			USequencerPlaylistItem_Sequence* SequenceItem = Cast<USequencerPlaylistItem_Sequence>(Item);
+			if (SequenceItem && SequenceItem->GetSequence())
+			{
+				PrecacheSequences.Add(SequenceItem->GetSequence());
+			}
+		}
+	}
+
+	FConcertSequencerPrecacheEvent PrecacheAddedEvent{ .bShouldBePrecached = true };
+	FConcertSequencerPrecacheEvent PrecacheRemovedEvent{ .bShouldBePrecached = false };
+
+	const TSet<TObjectPtr<ULevelSequence>> PrecacheUnion = PreviousPrecache.Union(PrecacheSequences);
+	for (const TObjectPtr<ULevelSequence>& Sequence : PrecacheUnion)
+	{
+		const FString SequenceObjectPath = Sequence->GetPathName();
+		if (!PreviousPrecache.Contains(Sequence))
+		{
+			UE_LOG(LogSequencerPlaylists, Verbose, TEXT("Adding sequence '%s' to precache set"), *SequenceObjectPath);
+			PrecacheAddedEvent.SequenceObjectPaths.Add(*SequenceObjectPath);
+		}
+		else if (!PrecacheSequences.Contains(Sequence))
+		{
+			UE_LOG(LogSequencerPlaylists, Verbose, TEXT("Removing sequence '%s' from precache set"), *SequenceObjectPath);
+			PrecacheRemovedEvent.SequenceObjectPaths.Add(*SequenceObjectPath);
+		}
+	}
+
+	if (TSharedPtr<IConcertClientSession> Session = WeakSession.Pin())
+	{
+		if (PrecacheAddedEvent.SequenceObjectPaths.Num() > 0)
+		{
+			Session->SendCustomEvent(PrecacheAddedEvent, Session->GetSessionServerEndpointId(), EConcertMessageFlags::ReliableOrdered);
+		}
+
+		if (PrecacheRemovedEvent.SequenceObjectPaths.Num() > 0)
+		{
+			Session->SendCustomEvent(PrecacheRemovedEvent, Session->GetSessionServerEndpointId(), EConcertMessageFlags::ReliableOrdered);
+		}
+	}
+}
+
+
+void USequencerPlaylistsSubsystem::OnConcertSessionStartup(TSharedRef<IConcertClientSession> InSession)
+{
+	WeakSession = InSession;
+
+	FConcertSequencerPrecacheEvent PrecacheEvent{ .bShouldBePrecached = true };
+
+	for (const TObjectPtr<ULevelSequence>& Sequence : PrecacheSequences)
+	{
+		PrecacheEvent.SequenceObjectPaths.Add(*Sequence->GetPathName());
+	}
+
+	if (PrecacheEvent.SequenceObjectPaths.Num() > 0)
+	{
+		InSession->SendCustomEvent(PrecacheEvent, InSession->GetSessionServerEndpointId(), EConcertMessageFlags::ReliableOrdered);
+	}
+}
+
+
+void USequencerPlaylistsSubsystem::OnConcertSessionShutdown(TSharedRef<IConcertClientSession> InSession)
+{
+	ensure(WeakSession == InSession);
+	WeakSession = nullptr;
 }
diff --git a/Engine/Plugins/VirtualProduction/SequencerPlaylists/Source/SequencerPlaylists/Private/SequencerPlaylistsSubsystem.h b/Engine/Plugins/VirtualProduction/SequencerPlaylists/Source/SequencerPlaylists/Private/SequencerPlaylistsSubsystem.h
--- a/Engine/Plugins/VirtualProduction/SequencerPlaylists/Source/SequencerPlaylists/Private/SequencerPlaylistsSubsystem.h
+++ b/Engine/Plugins/VirtualProduction/SequencerPlaylists/Source/SequencerPlaylists/Private/SequencerPlaylistsSubsystem.h
@@ -6,7 +6,9 @@
 #include "SequencerPlaylistsSubsystem.generated.h"
 
 
+class IConcertClientSession;
 class SSequencerPlaylistPanel;
+class ULevelSequence;
 class USequencerPlaylist;
 class USequencerPlaylistPlayer;
 
@@ -52,9 +54,14 @@ public:
 	USequencerPlaylistPlayer* CreatePlayerForEditor(TSharedRef<SSequencerPlaylistPanel> Editor);
 	void NotifyEditorClosed(SSequencerPlaylistPanel* Editor);
 
+	void UpdatePrecacheSet();
+
 protected:
 	USequencerPlaylist* CreateTransientPlaylistForEditor(TSharedRef<SSequencerPlaylistPanel> Editor);
 
+	void OnConcertSessionStartup(TSharedRef<IConcertClientSession> InSession);
+	void OnConcertSessionShutdown(TSharedRef<IConcertClientSession> InSession);
+
 protected:
 	UPROPERTY(Transient)
 	TMap<FSequencerPlaylistEditorHandle, TObjectPtr<UPackage>> EditorPackages;
@@ -64,4 +71,10 @@ protected:
 
 	UPROPERTY(Transient)
 	TMap<FSequencerPlaylistEditorHandle, TObjectPtr<USequencerPlaylist>> EditorPlaylists;
+
+	UPROPERTY(Transient)
+	TSet<TObjectPtr<ULevelSequence>> PrecacheSequences;
+
+	/** Weak pointer to the client session with which to send events. May be null or stale. */
+	TWeakPtr<IConcertClientSession> WeakSession;
 };
diff --git a/Engine/Plugins/VirtualProduction/SequencerPlaylists/Source/SequencerPlaylists/Private/SequencerPlaylistsWidgets.cpp b/Engine/Plugins/VirtualProduction/SequencerPlaylists/Source/SequencerPlaylists/Private/SequencerPlaylistsWidgets.cpp
--- a/Engine/Plugins/VirtualProduction/SequencerPlaylists/Source/SequencerPlaylists/Private/SequencerPlaylistsWidgets.cpp
+++ b/Engine/Plugins/VirtualProduction/SequencerPlaylists/Source/SequencerPlaylists/Private/SequencerPlaylistsWidgets.cpp
@@ -765,6 +765,11 @@ void SSequencerPlaylistPanel::LoadPlaylist(USequencerPlaylist* PlaylistToLoad)
 	{
 		Playlist->GetPackage()->ClearDirtyFlag();
 	}
+
+	if (USequencerPlaylistsSubsystem* Subsystem = GEditor->GetEditorSubsystem<USequencerPlaylistsSubsystem>())
+	{
+		Subsystem->UpdatePrecacheSet();
+	}
 }
 
 
@@ -1362,7 +1367,7 @@ TSharedRef<SWidget> SSequencerPlaylistItemWidget::GenerateWidgetForColumn(const
 			CastChecked<USequencerPlaylistItem_Sequence>(Item);
 
 		TSharedPtr<ISinglePropertyView> PropView = PropertyEditorModule.CreateSingleProperty(SequenceItem,
-			GET_MEMBER_NAME_CHECKED(USequencerPlaylistItem_Sequence, Sequence),
+			USequencerPlaylistItem_Sequence::GetSequencePropertyName(),
 			SinglePropParams);
 
 		return SNew(SBox)
@@ -1650,9 +1655,9 @@ TOptional<int32> SSequencerPlaylistItemWidget::GetItemLengthDisplayFrames() cons
 {
 	if (const USequencerPlaylistItem_Sequence* SequenceItem = Cast<const USequencerPlaylistItem_Sequence>(GetItem()))
 	{
-		if (SequenceItem->Sequence && SequenceItem->Sequence->GetMovieScene())
+		if (SequenceItem->GetSequence() && SequenceItem->GetSequence()->GetMovieScene())
 		{
-			UMovieScene* MovieScene = SequenceItem->Sequence->GetMovieScene();
+			UMovieScene* MovieScene = SequenceItem->GetSequence()->GetMovieScene();
 			TRange<FFrameNumber> Range = MovieScene->GetPlaybackRange();
 			if (Range.GetLowerBound().IsClosed() && Range.GetUpperBound().IsClosed())
 			{
diff --git a/Engine/Plugins/VirtualProduction/SequencerPlaylists/Source/SequencerPlaylists/SequencerPlaylists.Build.cs b/Engine/Plugins/VirtualProduction/SequencerPlaylists/Source/SequencerPlaylists/SequencerPlaylists.Build.cs
--- a/Engine/Plugins/VirtualProduction/SequencerPlaylists/Source/SequencerPlaylists/SequencerPlaylists.Build.cs
+++ b/Engine/Plugins/VirtualProduction/SequencerPlaylists/Source/SequencerPlaylists/SequencerPlaylists.Build.cs
@@ -31,6 +31,9 @@ public class SequencerPlaylists : ModuleRules
 			"ToolWidgets",
 			"MovieSceneTools",
 			"EditorSubsystem",
+			"ConcertClient",
+			"ConcertSyncClient",
+			"ConcertSyncCore",
 		});
 
 		DynamicallyLoadedModuleNames.AddRange(new string[] {
