Fix pause/hold usage in tool tip for Playlists tool.

diff --git a/Engine/Plugins/VirtualProduction/SequencerPlaylists/Source/SequencerPlaylists/Private/SequencerPlaylistsWidgets.cpp b/Engine/Plugins/VirtualProduction/SequencerPlaylists/Source/SequencerPlaylists/Private/SequencerPlaylistsWidgets.cpp
--- a/Engine/Plugins/VirtualProduction/SequencerPlaylists/Source/SequencerPlaylists/Private/SequencerPlaylistsWidgets.cpp
+++ b/Engine/Plugins/VirtualProduction/SequencerPlaylists/Source/SequencerPlaylists/Private/SequencerPlaylistsWidgets.cpp
@@ -1431,7 +1431,7 @@ TSharedRef<SWidget> SSequencerPlaylistItemWidget::GenerateWidgetForColumn(const
 		TSharedRef<SWidget> PauseToggle = SNew(SCheckBox)
 			.Padding(FMargin(4.0f, 2.0f))
 			.HAlign(HAlign_Center)
-			.ToolTipText(LOCTEXT("TogglePauseTooltip", "Enable or disable hold. Hold will infinitely hold the first frame of this item until manually played. Items are put into a hold state at the start of a take, or manually by hitting \"Reset.\""))
+			.ToolTipText(LOCTEXT("TogglePauseTooltip", "Enable or disable pause. Pause will infinitely hold the first frame of this item until manually played. Items are put into a pause state at the start of a take, or manually by hitting \"Reset.\""))
 			.Style(FAppStyle::Get(), "ToggleButtonCheckbox")
 			.IsChecked_Lambda([this]() { return IsPaused.Get() ? ECheckBoxState::Checked : ECheckBoxState::Unchecked; })
 			.OnCheckStateChanged_Lambda([this](ECheckBoxState InState) { PauseClickedDelegate.Execute(*this); })
