Use the transient package for playlists

#rnx

diff --git a/Engine/Plugins/VirtualProduction/SequencerPlaylists/Source/SequencerPlaylists/Private/SequencerPlaylistsSubsystem.cpp b/Engine/Plugins/VirtualProduction/SequencerPlaylists/Source/SequencerPlaylists/Private/SequencerPlaylistsSubsystem.cpp
--- a/Engine/Plugins/VirtualProduction/SequencerPlaylists/Source/SequencerPlaylists/Private/SequencerPlaylistsSubsystem.cpp
+++ b/Engine/Plugins/VirtualProduction/SequencerPlaylists/Source/SequencerPlaylists/Private/SequencerPlaylistsSubsystem.cpp
@@ -73,7 +73,7 @@ USequencerPlaylistPlayer* USequencerPlaylistsSubsystem::CreatePlayerForEditor(TS
 
 USequencerPlaylist* USequencerPlaylistsSubsystem::CreateTransientPlaylistForEditor(TSharedRef<SSequencerPlaylistPanel> Editor)
 {
-	FName PackageName = ::MakeUniqueObjectName(nullptr, UPackage::StaticClass(), TEXT("/Engine/Transient/SequencerPlaylist"));
+	FName PackageName = ::MakeUniqueObjectName(GetTransientPackage(), UPackage::StaticClass(), TEXT("SequencerPlaylist"));
 	UPackage* PlaylistPackage = NewObject<UPackage>(nullptr, PackageName, RF_Transient | RF_Transactional);
 	return NewObject<USequencerPlaylist>(PlaylistPackage, TEXT("UntitledPlaylist"), RF_Transactional);
 }
