Fix display of preload column when Multi-user session is not active

#rnx

diff --git a/Engine/Plugins/VirtualProduction/SequencerPlaylists/Source/SequencerPlaylists/Private/SequencerPlaylistsSubsystem.h b/Engine/Plugins/VirtualProduction/SequencerPlaylists/Source/SequencerPlaylists/Private/SequencerPlaylistsSubsystem.h
--- a/Engine/Plugins/VirtualProduction/SequencerPlaylists/Source/SequencerPlaylists/Private/SequencerPlaylistsSubsystem.h
+++ b/Engine/Plugins/VirtualProduction/SequencerPlaylists/Source/SequencerPlaylists/Private/SequencerPlaylistsSubsystem.h
@@ -68,6 +68,10 @@ public:
 	 */
 	TPair<EConcertSequencerPreloadStatus, FText> GetPreloadStatusForSequence(const FTopLevelAssetPath& SequencePath);
 
+	TWeakPtr<IConcertClientSession> GetWeakSession() const
+	{
+		return WeakSession;
+	}
 protected:
 	USequencerPlaylist* CreateTransientPlaylistForEditor(TSharedRef<SSequencerPlaylistPanel> Editor);
 
diff --git a/Engine/Plugins/VirtualProduction/SequencerPlaylists/Source/SequencerPlaylists/Private/SequencerPlaylistsWidgets.cpp b/Engine/Plugins/VirtualProduction/SequencerPlaylists/Source/SequencerPlaylists/Private/SequencerPlaylistsWidgets.cpp
--- a/Engine/Plugins/VirtualProduction/SequencerPlaylists/Source/SequencerPlaylists/Private/SequencerPlaylistsWidgets.cpp
+++ b/Engine/Plugins/VirtualProduction/SequencerPlaylists/Source/SequencerPlaylists/Private/SequencerPlaylistsWidgets.cpp
@@ -449,7 +449,7 @@ TSharedRef<SWidget> SSequencerPlaylistPanel::Construct_AddSearchRow()
 
 TSharedRef<SWidget> SSequencerPlaylistPanel::Construct_ItemListView()
 {
-	return SAssignNew(ItemListView, SListView<TSharedPtr<FSequencerPlaylistRowData>>)
+	TSharedRef<SWidget> Widget = SAssignNew(ItemListView, SListView<TSharedPtr<FSequencerPlaylistRowData>>)
 		.SelectionMode(ESelectionMode::Single) // See TODO in HandleAcceptDrop
 		.ListItemsSource(&ItemRows)
 		.OnGenerateRow_Lambda([this](TSharedPtr<FSequencerPlaylistRowData> InData, const TSharedRef<STableViewBase>& OwnerTableView)
@@ -522,9 +522,15 @@ TSharedRef<SWidget> SSequencerPlaylistPanel::Construct_ItemListView()
 
 	if (ensure(HeaderRow))
 	{
-		const bool bPreloadVisible = IConcertSyncClientModule::Get().GetClient(TEXT("MultiUser")).IsValid();
+		bool bPreloadVisible = false;
+		if (TSharedPtr<IConcertSyncClient> ConcertSyncClient = IConcertSyncClientModule::Get().GetClient(TEXT("MultiUser")))
+		{
+			IConcertClientRef ConcertClient = ConcertSyncClient->GetConcertClient();
+			bPreloadVisible = ConcertClient->GetCurrentSession().IsValid();
+		}
 		HeaderRow->SetShowGeneratedColumn(ColumnName_Preload, bPreloadVisible);
 	}
+	return Widget;
 }
 
 
@@ -1605,7 +1611,8 @@ TSharedRef<SWidget> SSequencerPlaylistItemWidget::GenerateWidgetForColumn(const
 				.ColorAndOpacity_Lambda([this, WeakItem]() {
 					USequencerPlaylistsSubsystem* Subsystem = GEditor->GetEditorSubsystem<USequencerPlaylistsSubsystem>();
 					USequencerPlaylistItem_Sequence* SeqItem = Cast<USequencerPlaylistItem_Sequence>(WeakItem);
-					if (Subsystem && SeqItem)
+					const bool bHasSession = Subsystem ? Subsystem->GetWeakSession().IsValid() : false;
+					if (bHasSession && SeqItem)
 					{
 						const FTopLevelAssetPath SeqPath(SeqItem->GetSequence());
 						TPair<EConcertSequencerPreloadStatus, FText> Status = Subsystem->GetPreloadStatusForSequence(SeqPath);
@@ -1621,7 +1628,8 @@ TSharedRef<SWidget> SSequencerPlaylistItemWidget::GenerateWidgetForColumn(const
 				.ToolTipText_Lambda([this, WeakItem]() {
 					USequencerPlaylistsSubsystem* Subsystem = GEditor->GetEditorSubsystem<USequencerPlaylistsSubsystem>();
 					USequencerPlaylistItem_Sequence* SeqItem = Cast<USequencerPlaylistItem_Sequence>(WeakItem);
-					if (Subsystem && SeqItem)
+					const bool bHasSession = Subsystem ? Subsystem->GetWeakSession().IsValid() : false;
+					if (bHasSession && SeqItem)
 					{
 						const FTopLevelAssetPath SeqPath(SeqItem->GetSequence());
 						TPair<EConcertSequencerPreloadStatus, FText> Status = Subsystem->GetPreloadStatusForSequence(SeqPath);
